Chaos Mover: Async implementation of movement simulation, so it can run on the physics thread. This introduces a new backend for async, an FSimulation class that only knows about async safe objects, and a copy of the Movement State Machine that only uses async safe parameters. The state machine removes the use of movement modifiers which we might incorporate into movement modes to simplify.

diff --git a/Engine/Plugins/Experimental/ChaosMover/ChaosMover.uplugin b/Engine/Plugins/Experimental/ChaosMover/ChaosMover.uplugin
new file mode 100644
--- /dev/null
+++ b/Engine/Plugins/Experimental/ChaosMover/ChaosMover.uplugin
@@ -0,0 +1,30 @@
+{
+	"FileVersion": 1,
+	"Version": 1,
+	"VersionName": "1.0",
+	"FriendlyName": "Chaos Mover",
+	"Description": "Chaos Mover is an Unreal Engine plugin to support Chaos physics based character movement using the Mover framework.",
+	"Category": "Physics",
+	"CreatedBy": "Epic Games, Inc.",
+	"CreatedByURL": "https://epicgames.com",
+	"DocsURL": "",
+	"MarketplaceURL": "",
+	"SupportURL": "",
+	"CanContainContent": true,
+	"IsBetaVersion": false,
+	"IsExperimentalVersion": true,
+	"Installed": false,
+	"Modules": [
+		{
+			"Name": "ChaosMover",
+			"Type": "Runtime",
+			"LoadingPhase": "Default"
+		}
+	],
+	"Plugins": [
+		{
+			"Name": "Mover",
+			"Enabled": true
+		}
+	]
+}
\ No newline at end of file
diff --git a/Engine/Plugins/Experimental/ChaosMover/Source/ChaosMover/ChaosMover.Build.cs b/Engine/Plugins/Experimental/ChaosMover/Source/ChaosMover/ChaosMover.Build.cs
new file mode 100644
--- /dev/null
+++ b/Engine/Plugins/Experimental/ChaosMover/Source/ChaosMover/ChaosMover.Build.cs
@@ -0,0 +1,33 @@
+// Copyright Epic Games, Inc. All Rights Reserved.
+
+namespace UnrealBuildTool.Rules
+{
+	public class ChaosMover : ModuleRules
+	{
+		public ChaosMover(ReadOnlyTargetRules Target) : base(Target)
+		{
+			PublicDependencyModuleNames.AddRange(
+				new string[]
+				{
+					"Core",
+					"GameplayTags",
+					"InputCore",
+					"NetCore",
+					"Mover",
+					"Water"
+				}
+			);
+
+			PrivateDependencyModuleNames.AddRange(
+				new string[]
+				{
+					"CoreUObject",
+					"DeveloperSettings",
+					"Engine"
+				}
+			);
+
+			SetupModulePhysicsSupport(Target);
+		}
+	}
+}
\ No newline at end of file
diff --git a/Engine/Plugins/Experimental/ChaosMover/Source/ChaosMover/Private/Backends/ChaosMoverAsyncCallback.cpp b/Engine/Plugins/Experimental/ChaosMover/Source/ChaosMover/Private/Backends/ChaosMoverAsyncCallback.cpp
new file mode 100644
--- /dev/null
+++ b/Engine/Plugins/Experimental/ChaosMover/Source/ChaosMover/Private/Backends/ChaosMoverAsyncCallback.cpp
@@ -0,0 +1,66 @@
+// Copyright Epic Games, Inc. All Rights Reserved.
+
+#include "ChaosMoverAsyncCallback.h"
+
+#include "Chaos/Framework/Parallel.h"
+#include "ChaosMover/Backends/ChaosMoverBackend.h"
+#include "ChaosMover/ChaosMoverDeveloperSettings.h"
+#include "ChaosMover/ChaosMoverSimulation.h"
+#include "ChaosMover/ChaosMoverSimulationTypes.h"
+#include "PBDRigidsSolver.h"
+
+namespace UE::ChaosMover
+{
+	FAsyncCallback::~FAsyncCallback()
+	{
+		Simulations.Empty();
+	}
+
+	void FAsyncCallback::OnPreSimulate_Internal()
+	{
+		const FAsyncCallbackInput* AsyncInput = GetConsumerInput_Internal();
+
+		if (!AsyncInput || AsyncInput->InputData.IsEmpty())
+		{
+			return;
+		}
+
+		FMoverTimeStep TimeStep;
+		TimeStep.BaseSimTimeMs = GetSimTime_Internal() * 1000.0f;
+		TimeStep.StepMs = GetDeltaTime_Internal() * 1000.0f;
+		if (PhysicsSolver)
+		{
+			TimeStep.ServerFrame = PhysicsSolver->GetCurrentFrame() + NetworkPhysicsTickOffset;
+			TimeStep.BaseSimTimeMs += NetworkPhysicsTickOffset * TimeStep.StepMs;
+			TimeStep.bIsResimulating = PhysicsSolver->GetEvolution()->IsResimming();
+		}
+
+		FAsyncCallbackOutput& AsyncOutput = GetProducerOutputData_Internal();
+		AsyncOutput.OutputData.SetNum(AsyncInput->InputData.Num());
+		AsyncOutput.TimeStep.SetNum(AsyncInput->InputData.Num());
+
+		auto LambdaParallelUpdate = [&TimeStep, &AsyncInput, &AsyncOutput, this](int32 Idx) {
+			AsyncOutput.TimeStep[Idx] = TimeStep;
+			if (UE::ChaosMover::FSimulation* Simulation = Simulations[Idx].Get())
+			{
+				Simulation->SimulationTick(TimeStep, AsyncInput->InputData[Idx], AsyncOutput.OutputData[Idx]);
+			}
+		};
+
+		const bool ForceSingleThread = UE::ChaosMover::CVars::bForceSingleThreadedPT;
+		Chaos::PhysicsParallelFor(Simulations.Num(), LambdaParallelUpdate, ForceSingleThread);
+	}
+
+	void FAsyncCallback::OnContactModification_Internal(Chaos::FCollisionContactModifier& Modifier)
+	{
+		auto LambdaParallelUpdate = [&Modifier, this](int32 Idx) {
+			if (UE::ChaosMover::FSimulation* Simulation = Simulations[Idx].Get())
+			{
+				Simulation->ModifyContacts(Modifier);
+			}
+		};
+
+		const bool ForceSingleThread = UE::ChaosMover::CVars::bForceSingleThreadedPT;
+		Chaos::PhysicsParallelFor(Simulations.Num(), LambdaParallelUpdate, ForceSingleThread);
+	}
+}
\ No newline at end of file
diff --git a/Engine/Plugins/Experimental/ChaosMover/Source/ChaosMover/Private/Backends/ChaosMoverAsyncCallback.h b/Engine/Plugins/Experimental/ChaosMover/Source/ChaosMover/Private/Backends/ChaosMoverAsyncCallback.h
new file mode 100644
--- /dev/null
+++ b/Engine/Plugins/Experimental/ChaosMover/Source/ChaosMover/Private/Backends/ChaosMoverAsyncCallback.h
@@ -0,0 +1,53 @@
+// Copyright Epic Games, Inc. All Rights Reserved.
+
+#pragma once
+
+#include "Chaos/SimCallbackObject.h"
+#include "ChaosMover/ChaosMoverSimulationTypes.h"
+
+class UChaosMoverBackendComponent;
+
+namespace UE::ChaosMover
+{
+	class FSimulation;
+
+	struct FAsyncCallbackInput : public Chaos::FSimCallbackInput
+	{
+		TArray<FSimulationInputData> InputData;
+		void Reset()
+		{
+			InputData.Empty();
+		}
+	};
+
+	struct FAsyncCallbackOutput : public Chaos::FSimCallbackOutput
+	{
+		TArray<FSimulationOutputData> OutputData;
+		TArray<FMoverTimeStep> TimeStep;
+		void Reset()
+		{
+			OutputData.Empty();
+			TimeStep.Empty();
+		}
+	};
+
+	class FAsyncCallback : public Chaos::TSimCallbackObject<
+		FAsyncCallbackInput,
+		FAsyncCallbackOutput,
+		Chaos::ESimCallbackOptions::Presimulate |
+		Chaos::ESimCallbackOptions::ContactModification |
+		Chaos::ESimCallbackOptions::Rewind>
+	{
+	public:
+		virtual ~FAsyncCallback();
+
+		TArray<TSharedPtr<UE::ChaosMover::FSimulation>> Simulations;
+		TArray<TWeakObjectPtr<UChaosMoverBackendComponent>> Backends;
+		Chaos::FPhysicsSolver* PhysicsSolver;
+		int32 NetworkPhysicsTickOffset = 0;
+
+	protected:
+		virtual void OnPreSimulate_Internal() override;
+		virtual void OnContactModification_Internal(Chaos::FCollisionContactModifier& Modifier) override;
+	};
+}
\ No newline at end of file
diff --git a/Engine/Plugins/Experimental/ChaosMover/Source/ChaosMover/Private/Backends/ChaosMoverBackend.cpp b/Engine/Plugins/Experimental/ChaosMover/Source/ChaosMover/Private/Backends/ChaosMoverBackend.cpp
new file mode 100644
--- /dev/null
+++ b/Engine/Plugins/Experimental/ChaosMover/Source/ChaosMover/Private/Backends/ChaosMoverBackend.cpp
@@ -0,0 +1,454 @@
+// Copyright Epic Games, Inc. All Rights Reserved.
+
+#include "ChaosMover/Backends/ChaosMoverBackend.h"
+
+#include "Backends/ChaosMoverSubsystem.h"
+#include "NetworkChaosMoverData.h"
+
+#include "ChaosMover/ChaosMoverSimulation.h"
+#include "Components/PrimitiveComponent.h"
+#include "Framework/Threading.h"
+#include "GameFramework/Pawn.h"
+#include "PBDRigidsSolver.h"
+#include "Physics/Experimental/PhysScene_Chaos.h"
+#include "Physics/NetworkPhysicsComponent.h"
+#include "PhysicsEngine/PhysicsObjectExternalInterface.h"
+#include "MovementModeStateMachine.h"
+#include "MoveLibrary/MovementMixer.h"
+#include "MoveLibrary/MoverBlackboard.h"
+
+#include UE_INLINE_GENERATED_CPP_BY_NAME(ChaosMoverBackend)
+
+UChaosMoverBackendComponent::UChaosMoverBackendComponent()
+{
+	PrimaryComponentTick.bCanEverTick = false;
+
+	bWantsInitializeComponent = true;
+	bAutoActivate = true;
+
+	if (const Chaos::FPhysicsSolver* Solver = GetPhysicsSolver())
+	{
+		bIsUsingAsyncPhysics = Solver->IsUsingAsyncResults();
+	}
+
+	if (Chaos::FPhysicsSolverBase::IsNetworkPhysicsPredictionEnabled())
+	{
+		SetIsReplicatedByDefault(true);
+	}
+}
+
+void UChaosMoverBackendComponent::InitializeComponent()
+{
+	Super::InitializeComponent();
+
+	// Create NetworkPhysicsComponent
+	if (Chaos::FPhysicsSolverBase::IsNetworkPhysicsPredictionEnabled() && bIsUsingAsyncPhysics)
+	{
+		NetworkPhysicsComponent = NewObject<UNetworkPhysicsComponent>(GetOwner(), TEXT("PhysMover_NetworkPhysicsComponent"));
+
+		// This isn't technically a DSO component, but set it net addressable as though it is
+		NetworkPhysicsComponent->SetNetAddressable();
+		NetworkPhysicsComponent->SetIsReplicated(true);
+		NetworkPhysicsComponent->RegisterComponent();
+		if (!NetworkPhysicsComponent->HasBeenInitialized())
+		{
+			NetworkPhysicsComponent->InitializeComponent();
+		}
+		NetworkPhysicsComponent->Activate(true);
+
+		// Register network data for recording and rewind/resim
+		NetworkPhysicsComponent->CreateDataHistory<UE::ChaosMover::FNetworkDataTraits>(this);
+
+		if (NetworkPhysicsComponent->HasServerWorld())
+		{
+			if (APawn* PawnOwner = Cast<APawn>(GetOwner()))
+			{
+				// When we're owned by a pawn, keep an eye on whether it's currently player-controlled or not
+				PawnOwner->ReceiveControllerChangedDelegate.AddUniqueDynamic(this, &ThisClass::HandleOwningPawnControllerChanged_Server);
+				HandleOwningPawnControllerChanged_Server(PawnOwner, nullptr, PawnOwner->Controller);
+			}
+			else
+			{
+				// If the owner isn't a pawn, there's no chance of player input happening, so inputs to the PT are always produced on the server
+				NetworkPhysicsComponent->SetIsRelayingLocalInputs(true);
+			}
+		}
+
+		bIsInputProducer = NetworkPhysicsComponent->IsLocallyControlled();
+	}
+
+	MoverBlackboard_Async = NewObject<UMoverBlackboard>(GetOwner(), TEXT("MoverBlackboard_Async"));
+	MovementMixer_Async = NewObject<UMovementMixer>(GetOwner(), TEXT("MovementMixer_Async"));
+}
+
+void UChaosMoverBackendComponent::UninitializeComponent()
+{
+	if (NetworkPhysicsComponent)
+	{
+		NetworkPhysicsComponent->RemoveDataHistory();
+		NetworkPhysicsComponent->DestroyComponent();
+		NetworkPhysicsComponent = nullptr;
+	}
+
+	if (MoverBlackboard_Async)
+	{
+		MoverBlackboard_Async = nullptr;
+	}
+
+	if (MovementMixer_Async)
+	{
+		MovementMixer_Async = nullptr;
+	}
+
+	Super::UninitializeComponent();
+}
+
+void UChaosMoverBackendComponent::HandleOwningPawnControllerChanged_Server(APawn* OwnerPawn, AController* OldController, AController* NewController)
+{
+	// Inputs for player-controlled pawns originate on the player's client, all others originate on the server
+	if (NetworkPhysicsComponent)
+	{
+		NetworkPhysicsComponent->SetIsRelayingLocalInputs(!OwnerPawn->IsPlayerControlled());
+		bIsInputProducer = NetworkPhysicsComponent->IsLocallyControlled();
+	}
+}
+
+void UChaosMoverBackendComponent::HandleUpdatedComponentPhysicsStateChanged(UPrimitiveComponent* ChangedComponent, EComponentPhysicsStateChange StateChange)
+{
+	if (StateChange == EComponentPhysicsStateChange::Destroyed)
+	{
+	}
+	else if (StateChange == EComponentPhysicsStateChange::Created)
+	{
+	}
+
+	InitSimulation();
+}
+
+Chaos::FPhysicsSolver* UChaosMoverBackendComponent::GetPhysicsSolver() const
+{
+	if (UWorld* World = GetWorld())
+	{
+		if (FPhysScene* Scene = World->GetPhysicsScene())
+		{
+			return Scene->GetSolver();
+		}
+	}
+	return nullptr;
+}
+
+UMoverComponent& UChaosMoverBackendComponent::GetMoverComponent() const
+{
+	return *GetOuterUMoverComponent();
+}
+
+Chaos::FPBDRigidParticle* UChaosMoverBackendComponent::GetControlledParticle() const
+{
+	if (IPhysicsComponent* PhysicsComponent = Cast<IPhysicsComponent>(GetMoverComponent().GetUpdatedComponent()))
+	{
+		if (Chaos::FPhysicsObject* PhysObj = PhysicsComponent->GetPhysicsObjectByName(NAME_None))
+		{
+			return FPhysicsObjectExternalInterface::LockRead(PhysObj)->GetRigidParticle(PhysObj);
+		}
+	}
+
+	return nullptr;
+}
+
+TSharedPtr<UE::ChaosMover::FSimulation> UChaosMoverBackendComponent::CreateSimulation()
+{
+	return MakeShared<UE::ChaosMover::FSimulation>();
+}
+
+void UChaosMoverBackendComponent::InitSimulation()
+{
+	SimOutputRecord.Clear();
+
+	UE::ChaosMover::FSimulationOutputData OutputData;
+	OutputData.SyncState.MovementMode = GetMoverComponent().StartingMovementMode;
+
+	FMoverTimeStep TimeStep;
+	if (Chaos::FPhysicsSolver* Solver = GetPhysicsSolver())
+	{
+		TimeStep.BaseSimTimeMs = Solver->GetPhysicsResultsTime_External() * 1000.0f;
+		TimeStep.ServerFrame = Solver->GetCurrentFrame();
+		TimeStep.StepMs = Solver->GetAsyncDeltaTime() * 1000.0f;
+	}
+
+	SimOutputRecord.Add(TimeStep, OutputData);
+
+	Simulation->SimBlackboardWeakPtr = TWeakObjectPtr<UMoverBlackboard>(MoverBlackboard_Async);
+	Simulation->MovementMixerWeakPtr =  TWeakObjectPtr<UMovementMixer>(MovementMixer_Async);
+
+	InitSimulationStateMachine();
+
+	Simulation->Init(OutputData.SyncState);
+}
+
+void UChaosMoverBackendComponent::InitSimulationStateMachine()
+{
+	if (!GetOwner())
+	{
+		return;
+	}
+
+	UMoverComponent& MoverComp = GetMoverComponent();
+
+	// Set up FSM and initial movement states
+
+	Simulation->MoverStateMachine.SetOwnerActorName(GetNameSafe(GetOwner()));
+	Simulation->MoverStateMachine.SetOwnerActorLocalNetRole(GetOwner()->GetLocalRole());
+
+	TObjectPtr<UImmediateMovementModeTransition> ImmediateMovementModeTransition = NewObject<UImmediateMovementModeTransition>(&MoverComp, TEXT("ImmediateMovementModeTransition"), RF_Transient);
+	TObjectPtr<UNullMovementMode> NullMovementMode = NewObject<UNullMovementMode>(&MoverComp, TEXT("NullMovementMode"), RF_Transient);
+
+	Simulation->MoverStateMachine.InitDefaultModesAndTransitions(ImmediateMovementModeTransition, NullMovementMode);
+
+	Simulation->MoverStateMachine.ClearAllMovementModes();
+	Simulation->MoverStateMachine.ClearAllGlobalTransitions();
+
+	bool bHasMatchingStartingState = false;
+
+	for (const TPair<FName, TObjectPtr<UBaseMovementMode>>& Element : MoverComp.MovementModes)
+	{
+		if (Element.Value.Get() == nullptr)
+		{
+			UE_LOG(LogMover, Warning, TEXT("Invalid Movement Mode type '%s' detected on %s. Mover actor will not function correctly."),
+				*Element.Key.ToString(), *GetNameSafe(GetOwner()));
+			continue;
+		}
+
+		Simulation->MoverStateMachine.RegisterMovementMode(Element.Key, Element.Value);
+
+		bHasMatchingStartingState |= (MoverComp.StartingMovementMode == Element.Key);
+	}
+
+	for (TObjectPtr<UBaseMovementModeTransition>& Transition : MoverComp.Transitions)
+	{
+		Simulation->MoverStateMachine.RegisterGlobalTransition(Transition);
+	}
+
+	if (!bHasMatchingStartingState)
+	{
+		UE_LOG(LogMover, Warning, TEXT("Invalid StartingMovementMode '%s' specified on %s. Mover actor will not function."),
+			*MoverComp.StartingMovementMode.ToString(), *GetNameSafe(GetOwner()));
+	}
+	else if (MoverComp.StartingMovementMode != NAME_None)
+	{
+		Simulation->MoverStateMachine.SetDefaultMode(MoverComp.StartingMovementMode);
+		Simulation->MoverStateMachine.QueueNextMode(MoverComp.StartingMovementMode);
+	}
+}
+
+void UChaosMoverBackendComponent::BeginPlay()
+{
+	Super::BeginPlay();
+
+	Simulation = CreateSimulation();
+	InitSimulation();
+
+	// Register with the world subsystem
+	if (UChaosMoverSubsystem* ChaosMoverSubsystem = UWorld::GetSubsystem<UChaosMoverSubsystem>(GetWorld()))
+	{
+		ChaosMoverSubsystem->Register(this);
+	}
+
+	// Register a callback to watch for component state changes
+	if (UPrimitiveComponent* PrimComp = Cast<UPrimitiveComponent>(GetMoverComponent().GetUpdatedComponent()))
+	{
+		PrimComp->OnComponentPhysicsStateChanged.AddUniqueDynamic(this, &ThisClass::HandleUpdatedComponentPhysicsStateChanged);
+	}
+}
+
+void UChaosMoverBackendComponent::EndPlay(const EEndPlayReason::Type EndPlayReason)
+{
+	Simulation.Reset();
+
+	if (UPrimitiveComponent* PrimComp = Cast<UPrimitiveComponent>(GetMoverComponent().GetUpdatedComponent()))
+	{
+		PrimComp->OnComponentPhysicsStateChanged.RemoveDynamic(this, &ThisClass::HandleUpdatedComponentPhysicsStateChanged);
+	}
+
+	if (UChaosMoverSubsystem* ChaosMoverSubsystem = UWorld::GetSubsystem<UChaosMoverSubsystem>(GetWorld()))
+	{
+		ChaosMoverSubsystem->Unregister(this);
+	}
+
+	Super::EndPlay(EndPlayReason);
+}
+
+FMoverTimeStep UChaosMoverBackendComponent::GetMoverTimeStep(double DeltaTimeInSeconds, int32 NetworkPhysicsTickOffset) const
+{
+	FMoverTimeStep TimeStep;
+	if (Chaos::FPhysicsSolver* Solver = GetPhysicsSolver())
+	{
+		TimeStep.ServerFrame = Solver->GetCurrentFrame() + NetworkPhysicsTickOffset;
+		TimeStep.StepMs = DeltaTimeInSeconds * 1000.0f;
+		TimeStep.BaseSimTimeMs = Solver->GetPhysicsResultsTime_External() * 1000.0f + NetworkPhysicsTickOffset * TimeStep.StepMs;
+		TimeStep.bIsResimulating = false;
+	}
+
+	return TimeStep;
+}
+
+void UChaosMoverBackendComponent::ProduceInputData(int32 PhysicsStep, int32 NumSteps, double DeltaTimeInSeconds, int32 NetworkPhysicsTickOffset, UE::ChaosMover::FSimulationInputData& InputData)
+{
+	Chaos::EnsureIsInGameThreadContext();
+
+	FMoverTimeStep TimeStep = GetMoverTimeStep(DeltaTimeInSeconds, NetworkPhysicsTickOffset);
+
+	if (bIsInputProducer)
+	{
+		GenerateInput(TimeStep, InputData);
+	}
+
+	// Cache the produced input on the simulation so that it can be written to the network data
+	// This happens before the async input is received
+	Simulation->InitNetInputData(InputData.InputCmd);
+
+	UMoverComponent& MoverComp = GetMoverComponent();
+
+	// Add default simulation input data
+	FChaosMoverSimulationDefaultInputs& SimInputs = InputData.InputParams.FindOrAddMutableDataByType<FChaosMoverSimulationDefaultInputs>();
+	SimInputs.Gravity = MoverComp.GetGravityAcceleration();
+	SimInputs.UpDir = MoverComp.GetUpDirection();
+	SimInputs.OwningActor = GetOwner();
+	if (UPrimitiveComponent* PrimComp = Cast<UPrimitiveComponent>(MoverComp.GetUpdatedComponent()))
+	{
+		PrimComp->InitSweepCollisionParams(SimInputs.CollisionQueryParams, SimInputs.CollisionResponseParams);
+		SimInputs.CollisionChannel = PrimComp->GetCollisionObjectType();
+		PrimComp->CalcBoundingCylinder(SimInputs.BoundingCylinderRadius, SimInputs.BoundingCylinderfHeight);
+	}
+	if (IPhysicsComponent* PhysComp = Cast<IPhysicsComponent>(MoverComp.GetUpdatedComponent()))
+	{
+		SimInputs.PhysicsObject = PhysComp->GetPhysicsObjectById(0); // Get the root physics object
+	}
+
+	if (MoverComp.OnPreSimulationTick.IsBound())
+	{
+		MoverComp.OnPreSimulationTick.Broadcast(TimeStep, InputData.InputCmd);
+	}
+}
+
+void UChaosMoverBackendComponent::GenerateInput(const FMoverTimeStep& TimeStep, UE::ChaosMover::FSimulationInputData& InputData)
+{
+	UMoverComponent& MoverComp = GetMoverComponent();
+	MoverComp.ProduceInput(TimeStep.StepMs, &InputData.InputCmd);
+}
+
+void UChaosMoverBackendComponent::ConsumeOutputData(const FMoverTimeStep& TimeStep, const UE::ChaosMover::FSimulationOutputData& OutputData)
+{
+	Chaos::EnsureIsInGameThreadContext();
+
+	SimOutputRecord.Add(TimeStep, OutputData);
+}
+
+void UChaosMoverBackendComponent::FinalizeFrame(double ResultsTimeInSeconds)
+{
+	Chaos::EnsureIsInGameThreadContext();
+
+	FMoverTimeStep TimeStep;
+	UE::ChaosMover::FSimulationOutputData InterpolatedOutput;
+	if (bIsUsingAsyncPhysics)
+	{
+		SimOutputRecord.GetInterpolated(ResultsTimeInSeconds, TimeStep, InterpolatedOutput);
+	}
+	else
+	{
+		InterpolatedOutput = SimOutputRecord.GetLatest();
+	}
+
+	// Physics interactions in the last frame may have caused a change in position or velocity that's different from what a simple lerp would predict,
+	// so stomp the lerped sync state's transform data with that of the actual particle after the last sim frame
+	FMoverDefaultSyncState& TransformSyncState = InterpolatedOutput.SyncState.SyncStateCollection.FindOrAddMutableDataByType<FMoverDefaultSyncState>();
+	if (Chaos::FPBDRigidParticle* Particle = GetControlledParticle())
+	{
+		TransformSyncState.SetTransforms_WorldSpace(Particle->GetX(), FRotator(Particle->GetR()), Particle->GetV(), TransformSyncState.GetMovementBase(), TransformSyncState.GetMovementBaseBoneName());
+
+		// Make sure the move direction intent is in base space (the base quat is identity if there's no base, effectively making this a no-op)
+		TransformSyncState.MoveDirectionIntent = TransformSyncState.GetCapturedMovementBaseQuat().UnrotateVector(TransformSyncState.MoveDirectionIntent);
+	}
+
+	// This eventually will process any events queued by the simulation during the async tick.
+	// For now it's just movement mode activation, deactivation and transition callbacks.
+	FlushModeChangeEvents();
+
+	UMoverComponent& MoverComp = GetMoverComponent();
+	MoverComp.SetSimulationOutput(TimeStep, InterpolatedOutput);
+	MoverComp.OnPostSimulationTick.Broadcast(TimeStep);
+}
+
+void UChaosMoverBackendComponent::ApplyNetInputData(const FMoverInputCmdContext& InNetInputCmd)
+{
+	if (Simulation)
+	{
+		Simulation->ApplyNetInputData(InNetInputCmd);
+	}
+}
+
+void UChaosMoverBackendComponent::BuildNetInputData(FMoverInputCmdContext& OutNetInputCmd) const
+{
+	if (Simulation)
+	{
+		Simulation->BuildNetInputData(OutNetInputCmd);
+	}
+}
+
+void UChaosMoverBackendComponent::ApplyNetStateData(const FMoverSyncState& InNetSyncState)
+{
+	if (Simulation)
+	{
+		Simulation->ApplyNetStateData(InNetSyncState);
+	}
+}
+
+void UChaosMoverBackendComponent::BuildNetStateData(FMoverSyncState& OutNetSyncState) const
+{
+	if (Simulation)
+	{
+		Simulation->BuildNetStateData(OutNetSyncState);
+	}
+}
+
+void UChaosMoverBackendComponent::RegisterSimulation() const
+{
+	if (UChaosMoverSubsystem* ChaosMoverSubsystem = UWorld::GetSubsystem<UChaosMoverSubsystem>(GetWorld()))
+	{
+		ChaosMoverSubsystem->RegisterSimulation(Simulation);
+	}
+}
+
+/** Flushes the queue of mode changed events, which causes all movement mode change callbacks to be called
+* This is not async safe, should only be called on the gameplay thread
+*/
+void UChaosMoverBackendComponent::FlushModeChangeEvents()
+{
+	if (Simulation)
+	{
+		const TArray<UE::ChaosMover::FMovementModeChangeEvent>& QueuedModeChangeEvents = Simulation->MoverStateMachine.GetQueuedModeChangeEvents();
+		UMoverComponent& MoverComp = GetMoverComponent();
+
+		for (int i = QueuedModeChangeEvents.Num() - 1; i >= 0; --i)
+		{
+			const UE::ChaosMover::FMovementModeChangeEvent& MovementModeChangeEvent = QueuedModeChangeEvents[i];
+			TStrongObjectPtr<UBaseMovementMode> PreviousMovementMode = MovementModeChangeEvent.PreviousMovementMode.Pin();
+			TStrongObjectPtr<UBaseMovementMode> NextMovementMode = MovementModeChangeEvent.NextMovementMode.Pin();
+
+			if (PreviousMovementMode)
+			{
+				PreviousMovementMode->Deactivate();
+			}
+
+			if (NextMovementMode)
+			{
+				NextMovementMode->Activate();
+			}
+
+			// Signal movement mode change event
+			MoverComp.OnMovementModeChanged.Broadcast(MovementModeChangeEvent.PreviousMovementModeName, MovementModeChangeEvent.NextMovementModeName);
+		}
+
+		Simulation->MoverStateMachine.ClearQueuedModeChangeEvents();
+	}
+}
+
diff --git a/Engine/Plugins/Experimental/ChaosMover/Source/ChaosMover/Private/Backends/ChaosMoverSubsystem.cpp b/Engine/Plugins/Experimental/ChaosMover/Source/ChaosMover/Private/Backends/ChaosMoverSubsystem.cpp
new file mode 100644
--- /dev/null
+++ b/Engine/Plugins/Experimental/ChaosMover/Source/ChaosMover/Private/Backends/ChaosMoverSubsystem.cpp
@@ -0,0 +1,208 @@
+// Copyright Epic Games, Inc. All Rights Reserved.
+
+#include "ChaosMoverSubsystem.h"
+
+#include "ChaosMoverAsyncCallback.h"
+
+#include "ChaosMover/Backends/ChaosMoverBackend.h"
+#include "ChaosMover/ChaosMoverDeveloperSettings.h"
+#include "GameFramework/PlayerController.h"
+#include "PBDRigidsSolver.h"
+#include "Physics/Experimental/PhysScene_Chaos.h"
+#include "Physics/NetworkPhysicsComponent.h"
+
+#include UE_INLINE_GENERATED_CPP_BY_NAME(ChaosMoverSubsystem)
+
+bool UChaosMoverSubsystem::DoesSupportWorldType(const EWorldType::Type WorldType) const
+{
+	return WorldType == EWorldType::Game || WorldType == EWorldType::PIE;
+}
+
+void UChaosMoverSubsystem::OnWorldBeginPlay(UWorld& InWorld)
+{
+	Super::OnWorldBeginPlay(InWorld);
+
+	if (FPhysScene* PhysScene = InWorld.GetPhysicsScene())
+	{
+		PhysScenePostTickCallbackHandle = PhysScene->OnPhysScenePostTick.AddUObject(this, &UChaosMoverSubsystem::OnPostPhysicsTick);
+
+		if (Chaos::FPhysicsSolver* Solver = PhysScene->GetSolver())
+		{
+			AsyncCallback = Solver->CreateAndRegisterSimCallbackObject_External<UE::ChaosMover::FAsyncCallback>();
+			AsyncCallback->PhysicsSolver = Solver;
+
+			if (FNetworkPhysicsCallback* SolverCallback = static_cast<FNetworkPhysicsCallback*>(Solver->GetRewindCallback()))
+			{
+				InjectInputsExternalCallbackHandle = SolverCallback->InjectInputsExternal.AddUObject(this, &UChaosMoverSubsystem::InjectInputs_External);
+			}
+		}
+	}
+}
+
+void UChaosMoverSubsystem::Deinitialize()
+{
+	if (UWorld* World = GetWorld())
+	{
+		if (FPhysScene* PhysScene = World->GetPhysicsScene())
+		{
+			if (Chaos::FPhysicsSolver* Solver = PhysScene->GetSolver())
+			{
+				if (InjectInputsExternalCallbackHandle.IsValid())
+				{
+					if (FNetworkPhysicsCallback* SolverCallback = static_cast<FNetworkPhysicsCallback*>(Solver->GetRewindCallback()))
+					{
+						SolverCallback->InjectInputsExternal.Remove(InjectInputsExternalCallbackHandle);
+					}
+				}
+
+				if (PhysScenePostTickCallbackHandle.IsValid())
+				{
+					PhysScene->OnPhysScenePostTick.Remove(PhysScenePostTickCallbackHandle);
+				}
+
+				if (AsyncCallback)
+				{
+					Solver->UnregisterAndFreeSimCallbackObject_External(AsyncCallback);
+				}
+			}
+		}
+	}
+
+	Super::Deinitialize();
+}
+
+void UChaosMoverSubsystem::Register(TWeakObjectPtr<UChaosMoverBackendComponent> InBackend)
+{
+	Backends.AddUnique(InBackend);
+	bBackendsListDirty = true;
+}
+
+void UChaosMoverSubsystem::Unregister(TWeakObjectPtr<UChaosMoverBackendComponent> InBackend)
+{
+	Backends.Remove(InBackend);
+	bBackendsListDirty = true;
+}
+
+void UChaosMoverSubsystem::RegisterSimulation(TSharedPtr<UE::ChaosMover::FSimulation> Simulation)
+{
+	AsyncCallback->Simulations.Add(Simulation);
+}
+
+void UChaosMoverSubsystem::UnregisterAllSimulations()
+{
+	AsyncCallback->Simulations.Empty();
+}
+
+void UChaosMoverSubsystem::InjectInputs_External(int32 PhysicsStep, int32 NumSteps)
+{
+	if (!AsyncCallback)
+	{
+		return;
+	}
+
+	// Clear invalid data before we attempt to use it.
+	bBackendsListDirty |= Backends.RemoveAllSwap([](const TWeakObjectPtr<UChaosMoverBackendComponent> Backend) {
+		return (Backend.IsValid() == false);
+		}) > 0;
+
+	// Copy a list of active backends to the callback
+	if (bBackendsListDirty)
+	{
+		AsyncCallback->Backends.Empty();
+		for (int Idx = 0; Idx < Backends.Num(); ++Idx)
+		{
+			Backends[Idx]->RegisterSimulation();
+			AsyncCallback->Backends.Add(Backends[Idx]);
+		}
+		bBackendsListDirty = false;
+	}
+
+	// Refresh the network physics tick offset in the async callback
+	const int32 NetPhysicsTickOffset = GetNetworkPhysicsTickOffset();
+	AsyncCallback->NetworkPhysicsTickOffset = NetPhysicsTickOffset;
+
+	UE::ChaosMover::FAsyncCallbackInput* AsyncInput = AsyncCallback->GetProducerInputData_External();
+	AsyncInput->Reset();
+	AsyncInput->InputData.SetNum(Backends.Num());
+
+	double DeltaSeconds = 0.0;
+	if (const UWorld* World = GetWorld())
+	{
+		if (const FPhysScene* Scene = World->GetPhysicsScene())
+		{
+			if (const Chaos::FPhysicsSolver* Solver = Scene->GetSolver())
+			{
+				if (Solver->IsUsingAsyncResults())
+				{
+					DeltaSeconds = Solver->GetAsyncDeltaTime();
+				}
+				else
+				{
+					DeltaSeconds = World->GetDeltaSeconds();
+				}
+			}
+		}
+	}
+
+	auto LambdaParallelUpdate = [NetPhysicsTickOffset, PhysicsStep, NumSteps, AsyncInput, DeltaSeconds, this](int32 Idx) {
+		Backends[Idx]->ProduceInputData(PhysicsStep, NumSteps, DeltaSeconds, NetPhysicsTickOffset, AsyncInput->InputData[Idx]);
+	};
+
+	const bool ForceSingleThread = UE::ChaosMover::CVars::bForceSingleThreadedGT;
+	Chaos::PhysicsParallelFor(Backends.Num(), LambdaParallelUpdate, ForceSingleThread);
+}
+
+void UChaosMoverSubsystem::OnPostPhysicsTick(FChaosScene* Scene)
+{
+	if (!AsyncCallback)
+	{
+		return;
+	}
+
+	const bool ForceSingleThread = UE::ChaosMover::CVars::bForceSingleThreadedGT;
+
+	// Pop each async output
+	while (Chaos::TSimCallbackOutputHandle<UE::ChaosMover::FAsyncCallbackOutput> AsyncOutput = AsyncCallback->PopFutureOutputData_External())
+	{
+		auto LambdaParallelUpdate = [&AsyncOutput, this](int32 Idx) {
+			if (UChaosMoverBackendComponent* Backend = AsyncCallback->Backends[Idx].Get())
+			{
+				Backend->ConsumeOutputData(AsyncOutput->TimeStep[Idx], AsyncOutput->OutputData[Idx]);
+			}
+		};
+
+		Chaos::PhysicsParallelFor(AsyncCallback->Backends.Num(), LambdaParallelUpdate, ForceSingleThread);
+	}
+
+	double ResultsTime = 0.0;
+	if (Scene != nullptr)
+	{
+		if (const Chaos::FPhysicsSolver* Solver = Scene->GetSolver())
+		{
+			ResultsTime = Solver->GetPhysicsResultsTime_External();
+		}
+	}
+
+	// Call finalize frame
+	auto LambdaParallelUpdate = [ResultsTime, this](int32 Idx) {
+		if (UChaosMoverBackendComponent* Backend = AsyncCallback->Backends[Idx].Get())
+		{
+			Backend->FinalizeFrame(ResultsTime);
+		}
+	};
+
+	Chaos::PhysicsParallelFor(AsyncCallback->Backends.Num(), LambdaParallelUpdate, ForceSingleThread);
+}
+
+int32 UChaosMoverSubsystem::GetNetworkPhysicsTickOffset() const
+{
+	if (UWorld* World = GetWorld())
+	{
+		if (APlayerController* PlayerController = World->GetFirstPlayerController())
+		{
+			return PlayerController->GetNetworkPhysicsTickOffset();
+		}
+	}
+
+	return 0;
+}
\ No newline at end of file
diff --git a/Engine/Plugins/Experimental/ChaosMover/Source/ChaosMover/Private/Backends/ChaosMoverSubsystem.h b/Engine/Plugins/Experimental/ChaosMover/Source/ChaosMover/Private/Backends/ChaosMoverSubsystem.h
new file mode 100644
--- /dev/null
+++ b/Engine/Plugins/Experimental/ChaosMover/Source/ChaosMover/Private/Backends/ChaosMoverSubsystem.h
@@ -0,0 +1,56 @@
+// Copyright Epic Games, Inc. All Rights Reserved.
+
+#pragma once
+
+#include "Subsystems/WorldSubsystem.h"
+#include "Chaos/Declares.h"
+
+#include "ChaosMoverSubsystem.generated.h"
+
+namespace UE::ChaosMover
+{
+	class FAsyncCallback;
+	class FSimulation;
+}
+
+class FChaosScene;
+class UChaosMoverBackendComponent;
+
+UCLASS()
+class UChaosMoverSubsystem : public UWorldSubsystem
+{
+	GENERATED_BODY()
+
+public:
+	void Register(TWeakObjectPtr<UChaosMoverBackendComponent> Backend);
+	void Unregister(TWeakObjectPtr<UChaosMoverBackendComponent> Backend);
+
+	/**
+	 * Registers a backend simulation's object shared pointer so it is kept alive during simulation,
+	 * even if its owner actor gets destroyed and garbage collected
+	 */ 
+	void RegisterSimulation(TSharedPtr<UE::ChaosMover::FSimulation> Simulation);
+
+protected:
+	virtual bool DoesSupportWorldType(const EWorldType::Type WorldType) const override;
+	virtual void OnWorldBeginPlay(UWorld& InWorld) override;
+	virtual void Deinitialize() override;
+
+	int32 GetNetworkPhysicsTickOffset() const;
+
+	void InjectInputs_External(int32 PhysicsStep, int32 NumSteps);
+	void OnPostPhysicsTick(FChaosScene* Scene);
+
+	/**
+	 * Unregisters all simulations. This is called internally when the list of backends is refreshed
+	 * The existing backends then all register their simulations again so they are kept alive
+	 */
+	void UnregisterAllSimulations();
+
+	TArray<TWeakObjectPtr<UChaosMoverBackendComponent>> Backends;
+	bool bBackendsListDirty = true;
+
+	FDelegateHandle InjectInputsExternalCallbackHandle;
+	FDelegateHandle PhysScenePostTickCallbackHandle;
+	class UE::ChaosMover::FAsyncCallback* AsyncCallback = nullptr;
+};
\ No newline at end of file
diff --git a/Engine/Plugins/Experimental/ChaosMover/Source/ChaosMover/Private/Backends/NetworkChaosMoverData.cpp b/Engine/Plugins/Experimental/ChaosMover/Source/ChaosMover/Private/Backends/NetworkChaosMoverData.cpp
new file mode 100644
--- /dev/null
+++ b/Engine/Plugins/Experimental/ChaosMover/Source/ChaosMover/Private/Backends/NetworkChaosMoverData.cpp
@@ -0,0 +1,140 @@
+// Copyright Epic Games, Inc. All Rights Reserved.
+
+#include "NetworkChaosMoverData.h"
+
+#include "ChaosMover/ChaosMoverSimulation.h"
+#include "ChaosMover/Backends/ChaosMoverBackend.h"
+
+void FNetworkChaosMoverInputData::ApplyData(UActorComponent* NetworkComponent) const
+{
+	if (UChaosMoverBackendComponent* BackendComp = Cast<UChaosMoverBackendComponent>(NetworkComponent))
+	{
+		BackendComp->ApplyNetInputData(InputCmdContext);
+	}
+}
+
+void FNetworkChaosMoverInputData::BuildData(const UActorComponent* NetworkComponent)
+{
+	if (const UChaosMoverBackendComponent* BackendComp = Cast<const UChaosMoverBackendComponent>(NetworkComponent))
+	{
+		BackendComp->BuildNetInputData(InputCmdContext);
+	}
+}
+
+void FNetworkChaosMoverInputData::DecayData(float DecayAmount)
+{
+	InputCmdContext.InputCollection.Decay(DecayAmount);
+}
+
+bool FNetworkChaosMoverInputData::NetSerialize(FArchive& Ar, class UPackageMap* PackageMap, bool& bOutSuccess)
+{
+	SerializeFrames(Ar);
+	if (PackageMap)
+	{
+		InputCmdContext.NetSerialize(FNetSerializeParams(Ar));
+		bOutSuccess = true;
+	}
+	else
+	{
+		bOutSuccess = false;
+	}
+
+	return bOutSuccess;
+}
+
+void FNetworkChaosMoverInputData::InterpolateData(const FNetworkPhysicsData& MinData, const FNetworkPhysicsData& MaxData)
+{
+	const FNetworkChaosMoverInputData& MinDataInput = static_cast<const FNetworkChaosMoverInputData&>(MinData);
+	const FNetworkChaosMoverInputData& MaxDataInput = static_cast<const FNetworkChaosMoverInputData&>(MaxData);
+
+	const float LerpFactor = (LocalFrame - MinDataInput.LocalFrame) / (MaxDataInput.LocalFrame - MinDataInput.LocalFrame);
+	InputCmdContext.InputCollection.Interpolate(MinDataInput.InputCmdContext.InputCollection, MaxDataInput.InputCmdContext.InputCollection, LerpFactor);
+}
+
+void FNetworkChaosMoverInputData::MergeData(const FNetworkPhysicsData& FromData)
+{
+	const FNetworkChaosMoverInputData& TypedFrom = static_cast<const FNetworkChaosMoverInputData&>(FromData);
+	InputCmdContext.InputCollection.Merge(TypedFrom.InputCmdContext.InputCollection);
+}
+
+void FNetworkChaosMoverInputData::ValidateData(const UActorComponent* NetworkComponent)
+{
+	// TODO
+}
+
+bool FNetworkChaosMoverInputData::CompareData(const FNetworkPhysicsData& PredictedData)
+{
+	const FMoverInputCmdContext& PredictedInputCmd = static_cast<const FNetworkChaosMoverInputData&>(PredictedData).InputCmdContext;
+	return !PredictedInputCmd.InputCollection.ShouldReconcile(InputCmdContext.InputCollection);
+}
+
+const FString FNetworkChaosMoverInputData::DebugData()
+{
+	FAnsiStringBuilderBase StringBuilder;
+	InputCmdContext.ToString(StringBuilder);
+	return FString::Printf(TEXT("FNetworkChaosMoverInputData:\n%hs"), StringBuilder.ToString());
+}
+
+//////////////////////////////////////////////////////////////////////////
+// FNetworkChaosMoverStateData
+
+void FNetworkChaosMoverStateData::ApplyData(UActorComponent* NetworkComponent) const
+{
+	if (NetworkComponent)
+	{
+		if (UChaosMoverBackendComponent* BackendComp = Cast<UChaosMoverBackendComponent>(NetworkComponent))
+		{
+			BackendComp->ApplyNetStateData(SyncState);
+		}
+	}
+}
+
+void FNetworkChaosMoverStateData::BuildData(const UActorComponent* NetworkComponent)
+{
+	if (NetworkComponent)
+	{
+		if (const UChaosMoverBackendComponent* BackendComp = Cast<const UChaosMoverBackendComponent>(NetworkComponent))
+		{
+			BackendComp->BuildNetStateData(SyncState);
+		}
+	}
+}
+
+bool FNetworkChaosMoverStateData::NetSerialize(FArchive& Ar, class UPackageMap* PackageMap, bool& bOutSuccess)
+{
+	SerializeFrames(Ar);
+	if (PackageMap)
+	{
+		FNetSerializeParams Params(Ar);
+		SyncState.NetSerialize(Params);
+		bOutSuccess = true;
+	}
+	else
+	{
+		bOutSuccess = false;
+	}
+
+	return bOutSuccess;
+}
+
+void FNetworkChaosMoverStateData::InterpolateData(const FNetworkPhysicsData& MinData, const FNetworkPhysicsData& MaxData)
+{
+	const FNetworkChaosMoverStateData& MinState = static_cast<const FNetworkChaosMoverStateData&>(MinData);
+	const FNetworkChaosMoverStateData& MaxState = static_cast<const FNetworkChaosMoverStateData&>(MaxData);
+
+	const float LerpFactor = (LocalFrame - MinState.LocalFrame) / (MaxState.LocalFrame - MinState.LocalFrame);
+	SyncState.Interpolate(&MinState.SyncState, &MaxState.SyncState, LerpFactor);
+}
+
+bool FNetworkChaosMoverStateData::CompareData(const FNetworkPhysicsData& PredictedData)
+{
+	const FMoverSyncState& PredictedSyncState = static_cast<const FNetworkChaosMoverStateData&>(PredictedData).SyncState;
+	return !PredictedSyncState.ShouldReconcile(SyncState);
+}
+
+const FString FNetworkChaosMoverStateData::DebugData()
+{
+	FAnsiStringBuilderBase StringBuilder;
+	SyncState.ToString(StringBuilder);
+	return FString::Printf(TEXT("FNetworkChaosMoverStateData:\n%hs"), StringBuilder.ToString());
+}
\ No newline at end of file
diff --git a/Engine/Plugins/Experimental/ChaosMover/Source/ChaosMover/Private/Backends/NetworkChaosMoverData.h b/Engine/Plugins/Experimental/ChaosMover/Source/ChaosMover/Private/Backends/NetworkChaosMoverData.h
new file mode 100644
--- /dev/null
+++ b/Engine/Plugins/Experimental/ChaosMover/Source/ChaosMover/Private/Backends/NetworkChaosMoverData.h
@@ -0,0 +1,96 @@
+// Copyright Epic Games, Inc. All Rights Reserved.
+
+#pragma once
+
+#include "MoverSimulationTypes.h"
+#include "Physics/NetworkPhysicsComponent.h"
+
+#include "NetworkChaosMoverData.generated.h"
+
+USTRUCT()
+struct FNetworkChaosMoverInputData : public FNetworkPhysicsData
+{
+	GENERATED_BODY()
+
+	UPROPERTY()
+	FMoverInputCmdContext InputCmdContext;
+
+	/**  Apply the data onto the network physics component */
+	virtual void ApplyData(UActorComponent* NetworkComponent) const override;
+
+	/**  Build the data from the network physics component */
+	virtual void BuildData(const UActorComponent* NetworkComponent) override;
+
+	/** Decay input during resimulation and forward prediction */
+	virtual void DecayData(float DecayAmount) override;
+
+	/**  Serialize data function that will be used to transfer the struct across the network */
+	bool NetSerialize(FArchive& Ar, class UPackageMap* Map, bool& bOutSuccess);
+
+	/** Interpolate the data in between two inputs data */
+	virtual void InterpolateData(const FNetworkPhysicsData& MinData, const FNetworkPhysicsData& MaxData) override;
+
+	/** Merge data into this input */
+	virtual void MergeData(const FNetworkPhysicsData& FromData) override;
+
+	/** Check input data is valid - Input is send from client to server, no need to make sure it's reasonable */
+	virtual void ValidateData(const UActorComponent* NetworkComponent) override;
+
+	virtual bool CompareData(const FNetworkPhysicsData& PredictedData) override;
+
+	/** Return string with debug information */
+	virtual const FString DebugData() override;
+};
+
+template<>
+struct TStructOpsTypeTraits<FNetworkChaosMoverInputData> : public TStructOpsTypeTraitsBase2<FNetworkChaosMoverInputData>
+{
+	enum
+	{
+		WithNetSerializer = true,
+	};
+};
+
+USTRUCT()
+struct FNetworkChaosMoverStateData : public FNetworkPhysicsData
+{
+	GENERATED_BODY()
+
+	UPROPERTY()
+	FMoverSyncState SyncState;
+
+	/**  Apply the data onto the network physics component */
+	virtual void ApplyData(UActorComponent* NetworkComponent) const override;
+
+	/**  Build the data from the network physics component */
+	virtual void BuildData(const UActorComponent* NetworkComponent) override;
+
+	/**  Serialize data function that will be used to transfer the struct across the network */
+	bool NetSerialize(FArchive& Ar, class UPackageMap* Map, bool& bOutSuccess);
+
+	/** Interpolate the data in between two inputs data */
+	virtual void InterpolateData(const FNetworkPhysicsData& MinData, const FNetworkPhysicsData& MaxData) override;
+
+	virtual bool CompareData(const FNetworkPhysicsData& PredictedData) override;
+
+	/** Return string with debug information */
+	virtual const FString DebugData() override;
+};
+
+template<>
+struct TStructOpsTypeTraits<FNetworkChaosMoverStateData> : public TStructOpsTypeTraitsBase2<FNetworkChaosMoverStateData>
+{
+	enum
+	{
+		WithNetSerializer = true,
+	};
+};
+
+namespace UE::ChaosMover
+{
+	struct FNetworkDataTraits
+	{
+		using InputsType = FNetworkChaosMoverInputData;
+		using StatesType = FNetworkChaosMoverStateData;
+	};
+}
\ No newline at end of file
diff --git a/Engine/Plugins/Experimental/ChaosMover/Source/ChaosMover/Private/ChaosMoverDeveloperSettings.cpp b/Engine/Plugins/Experimental/ChaosMover/Source/ChaosMover/Private/ChaosMoverDeveloperSettings.cpp
new file mode 100644
--- /dev/null
+++ b/Engine/Plugins/Experimental/ChaosMover/Source/ChaosMover/Private/ChaosMoverDeveloperSettings.cpp
@@ -0,0 +1,23 @@
+// Copyright Epic Games, Inc. All Rights Reserved.
+
+#include "ChaosMover/ChaosMoverDeveloperSettings.h"
+
+#include "HAL/IConsoleManager.h"
+
+namespace UE::ChaosMover
+{
+	namespace CVars
+	{
+		bool bForceSingleThreadedGT = false;
+		FAutoConsoleVariableRef CVarChaosMoverForceSingleThreadedGT(TEXT("ChaosMover.ForceSingleThreadedGT"),
+			bForceSingleThreadedGT, TEXT("Force updates on the game thread to be single threaded."));
+
+		bool bForceSingleThreadedPT = true;
+		FAutoConsoleVariableRef CVarChaosMoverForceSingleThreadedPT(TEXT("ChaosMover.ForceSingleThreadedPT"),
+			bForceSingleThreadedPT, TEXT("Force updates on the physics thread to be single threaded."));
+
+		bool bDrawGroundQueries = false;
+		FAutoConsoleVariableRef CVarChaosMoverDrawGroundQueries(TEXT("ChaosMover.DebugDraw.GroundQueries"),
+			bDrawGroundQueries, TEXT("Draw ground queries."));
+	}
+}
\ No newline at end of file
diff --git a/Engine/Plugins/Experimental/ChaosMover/Source/ChaosMover/Private/ChaosMoverLog.cpp b/Engine/Plugins/Experimental/ChaosMover/Source/ChaosMover/Private/ChaosMoverLog.cpp
new file mode 100644
--- /dev/null
+++ b/Engine/Plugins/Experimental/ChaosMover/Source/ChaosMover/Private/ChaosMoverLog.cpp
@@ -0,0 +1,5 @@
+// Copyright Epic Games, Inc. All Rights Reserved.
+
+#include "ChaosMover/ChaosMoverLog.h"
+
+DEFINE_LOG_CATEGORY(LogChaosMover);
\ No newline at end of file
diff --git a/Engine/Plugins/Experimental/ChaosMover/Source/ChaosMover/Private/ChaosMoverModule.cpp b/Engine/Plugins/Experimental/ChaosMover/Source/ChaosMover/Private/ChaosMoverModule.cpp
new file mode 100644
--- /dev/null
+++ b/Engine/Plugins/Experimental/ChaosMover/Source/ChaosMover/Private/ChaosMoverModule.cpp
@@ -0,0 +1,5 @@
+// Copyright Epic Games, Inc. All Rights Reserved.
+
+#include "Modules/ModuleManager.h"
+
+IMPLEMENT_MODULE(FDefaultModuleImpl, ChaosMover)
\ No newline at end of file
diff --git a/Engine/Plugins/Experimental/ChaosMover/Source/ChaosMover/Private/ChaosMoverSimulation.cpp b/Engine/Plugins/Experimental/ChaosMover/Source/ChaosMover/Private/ChaosMoverSimulation.cpp
new file mode 100644
--- /dev/null
+++ b/Engine/Plugins/Experimental/ChaosMover/Source/ChaosMover/Private/ChaosMoverSimulation.cpp
@@ -0,0 +1,117 @@
+// Copyright Epic Games, Inc. All Rights Reserved.
+
+#include "ChaosMover/ChaosMoverSimulation.h"
+
+#include "ChaosMover/ChaosMoverLog.h"
+#include "ChaosMover/ChaosMoverSimulationTypes.h"
+#include "MoverSimulationTypes.h"
+#include "Framework/Threading.h"
+#include "PBDRigidsSolver.h"
+#include "MovementModeStateMachine.h"
+#include "MoveLibrary/MovementMixer.h"
+#include "MoveLibrary/MoverBlackboard.h"
+#include "UObject/UObjectGlobals.h"
+#include "MoverDeveloperSettings.h"
+
+namespace UE::ChaosMover
+{
+	FSimulation::~FSimulation()
+	{
+	}
+	
+	void FSimulation::InitNetInputData(const FMoverInputCmdContext& InNetInputCmd)
+	{
+		InputCmd = InNetInputCmd;
+	}
+
+	void FSimulation::ApplyNetInputData(const FMoverInputCmdContext& InNetInputCmd)
+	{
+		InputCmd = InNetInputCmd;
+		bInputCmdOverridden = true;
+	}
+
+	void FSimulation::BuildNetInputData(FMoverInputCmdContext& OutNetInputCmd) const
+	{
+		OutNetInputCmd = InputCmd;
+	}
+
+	void FSimulation::ApplyNetStateData(const FMoverSyncState& InNetSyncState)
+	{
+		SyncState = InNetSyncState;
+	}
+
+	void FSimulation::BuildNetStateData(FMoverSyncState& OutNetSyncState) const
+	{
+		OutNetSyncState = SyncState;
+	}
+
+	void FSimulation::Init(const FMoverSyncState& InitialSyncState)
+	{
+		SyncState = InitialSyncState;
+		OnInit();
+	}
+
+	void FSimulation::OnInit()
+	{
+		
+	}
+
+	void FSimulation::SimulationTick(const FMoverTimeStep& TimeStep, const FSimulationInputData& InputData, FSimulationOutputData& OutputData)
+	{
+		Chaos::EnsureIsInPhysicsThreadContext();
+
+		OnPreSimulationTick(TimeStep, InputData);
+		OnSimulationTick(TimeStep, InputData, OutputData);
+		OnPostSimulationTick(TimeStep, OutputData);
+	}
+
+	void FSimulation::ModifyContacts(Chaos::FCollisionContactModifier& Modifier)
+	{
+		Chaos::EnsureIsInPhysicsThreadContext();
+
+		OnModifyContacts(Modifier);
+	}
+
+	void FSimulation::OnPreSimulationTick(const FMoverTimeStep& TimeStep, const FSimulationInputData& InputData)
+	{
+		if (bInputCmdOverridden)
+		{
+			InputData.InputCmd = InputCmd;
+		}
+	}
+
+	void FSimulation::OnSimulationTick(const FMoverTimeStep& TimeStep, const FSimulationInputData& InputData, FSimulationOutputData& OutputData)
+	{
+		TStrongObjectPtr<UMoverBlackboard> SimBlackboard = SimBlackboardWeakPtr.Pin();
+		TStrongObjectPtr<UMovementMixer> MovementMixer = MovementMixerWeakPtr.Pin();
+
+		if (!SimBlackboard || !MovementMixer)
+		{
+			// Looks like the movement component got deleted, no point in running the simulation
+			return;
+		}
+
+		// JK: I'm not sure this is the right thing to do, but we need a way to pass extra input data that exists in FSimulationInputData that is not InputCmd
+		// but still may be used by the simulation. I chose to use AuxState to pass that data for now, but I may rearrange things so we have a better pipe to send it through
+		// One the things we want to fix is that we shouldn't copy all that info around and write it once exactly where it is read from
+		FMoverAuxStateContext AuxState;
+		AuxState.AuxStateCollection = InputData.InputParams;
+		FMoverTickStartData TickStartData(InputData.InputCmd, SyncState, AuxState);
+
+		FMoverTickEndData TickEndData;
+		TickEndData.SyncState = SyncState;
+		MoverStateMachine.OnSimulationTick(TimeStep, TickStartData, SimBlackboard.Get(), MovementMixer.Get(), TickEndData);
+
+		OutputData.SyncState = TickEndData.SyncState;
+		OutputData.LastUsedInputCmd = InputData.InputCmd;
+	}
+
+	void FSimulation::OnPostSimulationTick(const FMoverTimeStep& TimeStep, const FSimulationOutputData& OutputData)
+	{
+	}
+
+	void FSimulation::OnModifyContacts(Chaos::FCollisionContactModifier& Modifier)
+	{
+	}
+
+} // namespace UE::ChaosMover
\ No newline at end of file
diff --git a/Engine/Plugins/Experimental/ChaosMover/Source/ChaosMover/Private/ChaosMoverSimulationTypes.cpp b/Engine/Plugins/Experimental/ChaosMover/Source/ChaosMover/Private/ChaosMoverSimulationTypes.cpp
new file mode 100644
--- /dev/null
+++ b/Engine/Plugins/Experimental/ChaosMover/Source/ChaosMover/Private/ChaosMoverSimulationTypes.cpp
@@ -0,0 +1,14 @@
+// Copyright Epic Games, Inc. All Rights Reserved.
+
+#include "ChaosMover/ChaosMoverSimulationTypes.h"
+
+void FChaosMoverSimulationDefaultInputs::Reset()
+{
+	CollisionResponseParams = FCollisionResponseParams();
+	CollisionQueryParams = FCollisionQueryParams();
+	UpDir = FVector::UpVector;
+	Gravity = -980.7 * UpDir;
+	PhysicsObject = nullptr;
+	OwningActor = nullptr;
+	CollisionChannel = ECC_Pawn;
+}
\ No newline at end of file
diff --git a/Engine/Plugins/Experimental/ChaosMover/Source/ChaosMover/Private/ChaosMoverStateMachine.cpp b/Engine/Plugins/Experimental/ChaosMover/Source/ChaosMover/Private/ChaosMoverStateMachine.cpp
new file mode 100644
--- /dev/null
+++ b/Engine/Plugins/Experimental/ChaosMover/Source/ChaosMover/Private/ChaosMoverStateMachine.cpp
@@ -0,0 +1,554 @@
+// Copyright Epic Games, Inc. All Rights Reserved.
+
+
+#include "ChaosMoverStateMachine.h"
+#include "Backends/ChaosMoverSubsystem.h"
+#include "MovementModeTransition.h"
+#include "MoverDeveloperSettings.h"
+#include "MoverLog.h"
+#include "MoveLibrary/MovementUtils.h"
+#include "MoverComponent.h"
+#include "MoveLibrary/MovementMixer.h"
+#include "Templates/SubclassOf.h"
+
+namespace UE::ChaosMover
+{
+	static int32 SkipGenerateMoveIfOverridden = 1;
+	FAutoConsoleVariableRef CVarSkipGenerateMoveIfOverridden(
+		TEXT("chaosmover.perf.SkipGenerateMoveIfOverridden"),
+		SkipGenerateMoveIfOverridden,
+		TEXT("If true and we have a layered move fully overriding movement, then we will skip calling OnGenerateMove on the active movement mode for better performance\n")
+	);
+
+	FMoverStateMachine::FMoverStateMachine()
+	{
+	}
+
+	FMoverStateMachine::~FMoverStateMachine()
+	{
+	}
+
+	void FMoverStateMachine::RegisterMovementMode(FName ModeName, TObjectPtr<UBaseMovementMode> Mode, bool bIsDefaultMode)
+	{
+		// JAH TODO: add validation and warnings for overwriting modes
+		// JAH TODO: add validation of Mode
+
+		Modes.Add(ModeName, TWeakObjectPtr<UBaseMovementMode>(Mode.Get()));
+
+		if (bIsDefaultMode)
+		{
+			//JAH TODO: add validation that we are only overriding the default null mode
+			DefaultModeName = ModeName;
+		}
+
+		Mode->OnRegistered(ModeName);
+	}
+
+	void FMoverStateMachine::UnregisterMovementMode(FName ModeName)
+	{
+		TWeakObjectPtr<UBaseMovementMode> ModeToUnregisterWeakPtr = Modes.FindAndRemoveChecked(ModeName);
+		TStrongObjectPtr<UBaseMovementMode> ModeToUnregister = ModeToUnregisterWeakPtr.Pin();
+
+		if (ModeToUnregister)
+		{
+			ModeToUnregister->OnUnregistered();
+		}
+	}
+
+	void FMoverStateMachine::ClearAllMovementModes()
+	{
+		Modes.Empty();
+
+		for (TPair<FName,TWeakObjectPtr<UBaseMovementMode>>& Element : Modes)
+		{
+			TStrongObjectPtr<UBaseMovementMode> Mode = Element.Value.Pin();
+
+			if (Mode)
+			{
+				Mode->OnUnregistered();
+			}
+		}
+
+		ConstructDefaultModes();	// Note that we're resetting to our defaults so we keep the null movement mode
+	}
+
+	void FMoverStateMachine::SetDefaultMode(FName NewDefaultModeName)
+	{
+		check(Modes.Contains(NewDefaultModeName));
+
+		DefaultModeName = NewDefaultModeName;
+	}
+
+	void FMoverStateMachine::RegisterGlobalTransition(TObjectPtr<UBaseMovementModeTransition> Transition)
+	{
+		GlobalTransitions.Add(TWeakObjectPtr<UBaseMovementModeTransition>(Transition.Get()));
+
+		Transition->OnRegistered();
+	}
+
+	void FMoverStateMachine::UnregisterGlobalTransition(TObjectPtr<UBaseMovementModeTransition> Transition)
+	{
+		Transition->OnUnregistered();
+
+		GlobalTransitions.Remove(TWeakObjectPtr<UBaseMovementModeTransition>(Transition.Get()));
+	}
+
+	void FMoverStateMachine::ClearAllGlobalTransitions()
+	{
+		for (TWeakObjectPtr<UBaseMovementModeTransition> TransitionWeakPtr : GlobalTransitions)
+		{
+			TStrongObjectPtr<UBaseMovementModeTransition> Transition = TransitionWeakPtr.Pin();
+			if (Transition)
+			{
+				Transition->OnUnregistered();
+			}
+		}
+
+		GlobalTransitions.Empty();
+	}
+
+	void FMoverStateMachine::QueueNextMode(FName DesiredNextModeName, bool bShouldReenter)
+	{
+		TStrongObjectPtr<UImmediateMovementModeTransition> QueuedModeTransition = QueuedModeTransitionWeakPtr.Pin();
+		if (QueuedModeTransition)
+		{
+			if (DesiredNextModeName != NAME_None)
+			{
+				const FName NextModeName = QueuedModeTransition->GetNextModeName();
+				const bool bShouldNextModeReenter = QueuedModeTransition->ShouldReenter();
+
+				if ((NextModeName != NAME_None) &&
+					(NextModeName != DesiredNextModeName || bShouldReenter != bShouldNextModeReenter))
+				{
+					UE_LOG(LogMover, Log, TEXT("%s (%s) Overwriting of queued mode change (%s, reenter: %i) with (%s, reenter: %i)"), *OwnerActorName, *UEnum::GetValueAsString(OwnerActorLocalNetRole), *NextModeName.ToString(), bShouldNextModeReenter, *DesiredNextModeName.ToString(), bShouldReenter);
+				}
+
+				if (Modes.Contains(DesiredNextModeName))
+				{
+					QueuedModeTransition->SetNextMode(DesiredNextModeName, bShouldReenter);
+				}
+				else
+				{
+					UE_LOG(LogMover, Warning, TEXT("Attempted to queue an unregistered movement mode: %s on owner %s"), *DesiredNextModeName.ToString(), *OwnerActorName);
+				}
+			}
+		}
+	}
+
+	void FMoverStateMachine::SetModeImmediately(FName DesiredModeName, bool bShouldReenter)
+	{
+		QueueNextMode(DesiredModeName, bShouldReenter);
+		AdvanceToNextMode();
+	}
+
+	void FMoverStateMachine::ClearQueuedMode()
+	{
+		TStrongObjectPtr<UImmediateMovementModeTransition> QueuedModeTransition = QueuedModeTransitionWeakPtr.Pin();
+		if (QueuedModeTransition)
+		{
+			QueuedModeTransition->Clear();
+		}
+	}
+
+	void FMoverStateMachine::OnSimulationTick(const FMoverTimeStep& TimeStep, const FMoverTickStartData& StartState, UMoverBlackboard* SimBlackboard, UMovementMixer* MovementMixer, FMoverTickEndData& OutputState)
+	{
+		if (!ensure(MovementMixer))
+		{
+			return;
+		}
+
+		const float EndingSimTimeMs = TimeStep.BaseSimTimeMs + TimeStep.StepMs;
+
+		//GEngine->AddOnScreenDebugMessage(-1, -0.1f, FColor::White, FString::Printf(TEXT("Mode Tick: %s  Queued: %s"), *CurrentModeName.ToString(), *NextModeName.ToString()));
+
+		FMoverTimeStep SubTimeStep = TimeStep;
+		FMoverTickStartData SubstepStartData = StartState;
+		FMoverDefaultSyncState* SubstepStartSyncState = SubstepStartData.SyncState.SyncStateCollection.FindMutableDataByType<FMoverDefaultSyncState>();
+
+		TStrongObjectPtr<UImmediateMovementModeTransition> QueuedModeTransition = QueuedModeTransitionWeakPtr.Pin();
+		if (QueuedModeTransition && !QueuedModeTransition->IsSet())
+		{
+			QueueNextMode(SubstepStartData.SyncState.MovementMode);
+		}
+
+		AdvanceToNextMode();
+
+		int SubStepCount = 0;
+		const int32 MaxConsecutiveFullRefundedSubsteps = GetDefault<UMoverDeveloperSettings>()->MaxTimesToRefundSubstep;
+		int32 NumConsecutiveFullRefundedSubsteps = 0;
+		while (SubTimeStep.BaseSimTimeMs < EndingSimTimeMs)
+		{
+			SubstepStartSyncState = SubstepStartData.SyncState.SyncStateCollection.FindMutableDataByType<FMoverDefaultSyncState>();
+			FMoverDefaultSyncState* OutputSyncState = &OutputState.SyncState.SyncStateCollection.FindOrAddMutableDataByType<FMoverDefaultSyncState>();
+			OutputState.SyncState.MovementMode = CurrentModeName;
+
+			OutputState.MovementEndState.ResetToDefaults();
+
+			SubTimeStep.StepMs = EndingSimTimeMs - SubTimeStep.BaseSimTimeMs;		// TODO: convert this to an overridable function that can support MaxStepTime, MaxIterations, etc.
+
+			// Transfer any queued moves into the starting state. They'll be started during the move generation.
+			FlushQueuedMovesToGroup(SubstepStartData.SyncState.LayeredMoves);
+			OutputState.SyncState.LayeredMoves = SubstepStartData.SyncState.LayeredMoves;
+
+			bool bModeSetFromInstantEffect = false;
+			// Apply any instant effects that were queued up between ticks
+			if (ApplyInstantEffects(SubstepStartData, SubTimeStep, OutputState.SyncState))
+			{
+				// Copying over our sync state collection to SubstepStartData so it is effectively the input sync state later for the movement mode. Doing this makes sure state modification from Instant Effects isn't overridden later by the movement mode
+				for (auto SyncDataIt = OutputState.SyncState.SyncStateCollection.GetCollectionDataIterator(); SyncDataIt; ++SyncDataIt)
+				{
+					if (SyncDataIt->Get())
+					{
+						SubstepStartData.SyncState.SyncStateCollection.AddOrOverwriteData(TSharedPtr<FMoverDataStructBase>(SyncDataIt->Get()->Clone()));
+					}
+				}
+
+				if (CurrentModeName != OutputState.SyncState.MovementMode)
+				{
+					bModeSetFromInstantEffect = true;
+					SetModeImmediately(OutputState.SyncState.MovementMode);
+					SubstepStartData.SyncState.MovementMode = CurrentModeName;
+				}
+			}
+
+			FMovementModifierGroup& CurrentModifiers = OutputState.SyncState.MovementModifiers;
+			TArray<TSharedPtr<FMovementModifierBase>> ActiveModifiers;
+			ensureMsgf(!CurrentModifiers.HasAnyMoves(), TEXT("Movement Modifiers are not supported in Mover async mode"));
+
+			FLayeredMoveGroup& CurrentLayeredMoves = OutputState.SyncState.LayeredMoves;
+
+			// Gather any layered move contributions
+			FProposedMove CombinedLayeredMove;
+			CombinedLayeredMove.MixMode = EMoveMixMode::AdditiveVelocity;
+			bool bHasLayeredMoveContributions = false;
+			MovementMixer->ResetMixerState();
+
+			TArray<TSharedPtr<FLayeredMoveBase>> ActiveMoves = CurrentLayeredMoves.GenerateActiveMoves_Async(SubTimeStep, SimBlackboard);
+
+			// Tick and accumulate all active moves
+			// Gather all proposed moves and distill this into a cumulative movement report. May include separate additive vs override moves.
+			// TODO: may want to sort by priority or other factors
+			for (TSharedPtr<FLayeredMoveBase>& ActiveMove : ActiveMoves)
+			{
+				FProposedMove MoveStep;
+				bool GenerateMoveResult = ActiveMove->GenerateMove_Async(SubstepStartData, SubTimeStep, SimBlackboard, MoveStep);
+				if (GenerateMoveResult)
+				{
+					// If this active move is already past it's first tick we don't need to set the preferred mode again
+					if (ActiveMove->StartSimTimeMs < SubTimeStep.BaseSimTimeMs)
+					{
+						MoveStep.PreferredMode = NAME_None;
+					}
+
+					bHasLayeredMoveContributions = true;
+					MovementMixer->MixLayeredMove(*ActiveMove, MoveStep, CombinedLayeredMove);
+				}
+			}
+
+			if (bHasLayeredMoveContributions && !CombinedLayeredMove.PreferredMode.IsNone() && !bModeSetFromInstantEffect)
+			{
+				SetModeImmediately(CombinedLayeredMove.PreferredMode);
+				OutputState.SyncState.MovementMode = CurrentModeName;
+			}
+
+			// Merge proposed movement from the current mode with movement from layered moves
+			if (!CurrentModeName.IsNone() && Modes.Contains(CurrentModeName))
+			{
+				UBaseMovementMode* CurrentMode = Modes[CurrentModeName].Get();
+				FProposedMove CombinedMove;
+				bool bHasModeMoveContribution = false;
+
+				if (!SkipGenerateMoveIfOverridden ||
+					!(bHasLayeredMoveContributions && CombinedLayeredMove.MixMode == EMoveMixMode::OverrideAll))
+				{
+					QUICK_SCOPE_CYCLE_COUNTER(STAT_GenerateMoveFromMode);
+					CurrentMode->GenerateMove(SubstepStartData, SubTimeStep, OUT CombinedMove);
+
+					bHasModeMoveContribution = true;
+				}
+
+				if (bHasModeMoveContribution && bHasLayeredMoveContributions)
+				{
+					MovementMixer->MixProposedMoves(CombinedLayeredMove, CombinedMove);
+				}
+				else if (bHasLayeredMoveContributions && !bHasModeMoveContribution)
+				{
+					CombinedMove = CombinedLayeredMove;
+				}
+
+				// Apply any layered move finish velocity settings
+				if (CurrentLayeredMoves.bApplyResidualVelocity)
+				{
+					CombinedMove.LinearVelocity = CurrentLayeredMoves.ResidualVelocity;
+				}
+				if (CurrentLayeredMoves.ResidualClamping >= 0.0f)
+				{
+					CombinedMove.LinearVelocity = CombinedMove.LinearVelocity.GetClampedToMaxSize(CurrentLayeredMoves.ResidualClamping);
+				}
+				CurrentLayeredMoves.ResetResidualVelocity();
+
+				// We need to replace this with some async equivalent (calling back to FSimulation? an optional FinalMoveProcessor object, a bit like the optional MoveMixer?)
+				// SyncTickParams->MoverComponent->ProcessGeneratedMovement.ExecuteIfBound(SubstepStartData, SubTimeStep, OUT CombinedMove);
+
+				// Execute the combined proposed move
+				{
+					FSimulationTickParams SimTickParams;
+					// SimTickParams.MovingComps is left empty in the async case, so we don't access resources used by the concurrent gameplay thread
+					SimTickParams.SimBlackboard = SimBlackboard;
+					SimTickParams.StartState = SubstepStartData;
+					SimTickParams.TimeStep = SubTimeStep;
+					SimTickParams.ProposedMove = CombinedMove;
+
+					// Check for any transitions, first those registered with the current movement mode, then global ones that could occur from any mode
+					FTransitionEvalResult EvalResult = FTransitionEvalResult::NoTransition;
+					TStrongObjectPtr<UBaseMovementModeTransition> TransitionToTrigger;
+
+					for (UBaseMovementModeTransition* Transition : CurrentMode->Transitions)
+					{
+						if (IsValid(Transition) && ((SubStepCount == 0) || !Transition->bFirstSubStepOnly))
+						{
+							EvalResult = Transition->Evaluate(SimTickParams);
+
+							if (!EvalResult.NextMode.IsNone())
+							{
+								if (EvalResult.NextMode != CurrentModeName || Transition->bAllowModeReentry)
+								{
+									TransitionToTrigger = TStrongObjectPtr<UBaseMovementModeTransition>(Transition);
+									break;
+								}
+							}
+						}
+					}
+
+					if (TransitionToTrigger == nullptr)
+					{
+						for (TWeakObjectPtr<UBaseMovementModeTransition> TransitionWeakPtr : GlobalTransitions)
+						{
+							TStrongObjectPtr<UBaseMovementModeTransition> Transition = TransitionWeakPtr.Pin();
+							if (Transition)
+							{
+								EvalResult = Transition->Evaluate(SimTickParams);
+
+								if (!EvalResult.NextMode.IsNone())
+								{
+									if (EvalResult.NextMode != CurrentModeName || Transition->bAllowModeReentry)
+									{
+										TransitionToTrigger = Transition;
+										break;
+									}
+								}
+							}
+						}
+					}
+
+					if (TransitionToTrigger && !EvalResult.NextMode.IsNone())
+					{
+						OutputState.MovementEndState.NextModeName = EvalResult.NextMode;
+						OutputState.MovementEndState.RemainingMs = SimTickParams.TimeStep.StepMs; 	// Pass all remaining time to next mode
+						TransitionToTrigger->Trigger(SimTickParams);
+					}
+					else
+					{
+						CurrentMode->SimulationTick(SimTickParams, OutputState);
+					}
+
+					OutputState.MovementEndState.RemainingMs = FMath::Clamp(OutputState.MovementEndState.RemainingMs, 0.0f, SubTimeStep.StepMs);
+				}
+
+				QueueNextMode(OutputState.MovementEndState.NextModeName);
+
+				// Check if all of the time for this Substep was refunded
+				if (FMath::IsNearlyEqual(SubTimeStep.StepMs, OutputState.MovementEndState.RemainingMs, UE_KINDA_SMALL_NUMBER))
+				{
+					NumConsecutiveFullRefundedSubsteps++;
+					// if we've done this sub step a lot before go ahead and just advance time to avoid freezing editor
+					if (NumConsecutiveFullRefundedSubsteps >= MaxConsecutiveFullRefundedSubsteps)
+					{
+						UE_LOG(LogMover, Warning, TEXT("Movement mode %s and %s on %s are stuck giving time back to each other. Overriding to advance to next substep."),
+							*CurrentModeName.ToString(),
+							*OutputState.MovementEndState.NextModeName.ToString(),
+							*OwnerActorName);
+						SubTimeStep.BaseSimTimeMs += SubTimeStep.StepMs;
+					}
+				}
+				else
+				{
+					NumConsecutiveFullRefundedSubsteps = 0;
+				}
+
+				//GEngine->AddOnScreenDebugMessage(-1, -0.1f, FColor::White, FString::Printf(TEXT("NextModeName: %s  Queued: %s"), *Output.MovementEndState.NextModeName.ToString(), *NextModeName.ToString()));
+			}
+
+			// Switch modes if necessary (note that this will allow exit/enter on the same state)
+			AdvanceToNextMode();
+			OutputState.SyncState.MovementMode = CurrentModeName;
+
+			ensureMsgf(ActiveModifiers.IsEmpty(), TEXT("Movement Modifiers are not supported in Async mode"));
+
+			const float RemainingMs = FMath::Clamp(OutputState.MovementEndState.RemainingMs, 0.0f, SubTimeStep.StepMs);
+			SubTimeStep.BaseSimTimeMs += (SubTimeStep.StepMs - RemainingMs);
+			SubTimeStep.StepMs = EndingSimTimeMs - SubTimeStep.BaseSimTimeMs;
+
+			SubstepStartData.SyncState = OutputState.SyncState;
+			SubstepStartData.AuxState = OutputState.AuxState;
+
+			++SubStepCount;
+		}
+
+		// Apply any instant effects that were queued up during this tick and didn't get handled in a substep
+		if (ApplyInstantEffects(SubstepStartData, SubTimeStep, OutputState.SyncState))
+		{
+			if (CurrentModeName != OutputState.SyncState.MovementMode)
+			{
+				SetModeImmediately(OutputState.SyncState.MovementMode);
+			}
+		}
+	}
+
+	void FMoverStateMachine::OnSimulationPreRollback(const FMoverSyncState* InvalidSyncState, const FMoverSyncState* SyncState, const FMoverAuxStateContext* InvalidAuxState, const FMoverAuxStateContext* AuxState)
+	{
+	}
+
+	void FMoverStateMachine::OnSimulationRollback(const FMoverSyncState* SyncState, const FMoverAuxStateContext* AuxState)
+	{
+		ClearQueuedMode();
+		QueuedLayeredMoves.Empty();
+		QueuedInstantEffects.Empty();
+	}
+
+
+	const UBaseMovementMode* FMoverStateMachine::GetCurrentMode() const
+	{
+		if (CurrentModeName != NAME_None && Modes.Contains(CurrentModeName))
+		{
+			return Modes[CurrentModeName].Get();
+		}
+
+		return nullptr;
+	}
+
+	const UBaseMovementMode* FMoverStateMachine::FindMovementMode(FName ModeName) const
+	{
+		if (ModeName != NAME_None && Modes.Contains(ModeName))
+		{
+			return Modes[ModeName].Get();
+		}
+
+		return nullptr;
+	}
+
+	void FMoverStateMachine::QueueLayeredMove(TSharedPtr<FLayeredMoveBase> Move)
+	{
+		QueuedLayeredMoves.Add(Move);
+	}
+
+	void FMoverStateMachine::QueueInstantMovementEffect(TSharedPtr<FInstantMovementEffect> Effect)
+	{
+		QueuedInstantEffects.Add(Effect);
+	}
+
+	void FMoverStateMachine::ConstructDefaultModes()
+	{
+		RegisterMovementMode(UNullMovementMode::NullModeName, TObjectPtr<UBaseMovementMode>(NullMovementModeWeakPtr.Get()), /*bIsDefaultMode =*/ true);
+		DefaultModeName = NAME_None;
+		CurrentModeName = UNullMovementMode::NullModeName;
+
+		QueuedModeTransitionWeakPtr = ImmediateMovementModeTransitionWeakPtr;
+
+		ClearQueuedMode();
+	}
+
+	void FMoverStateMachine::InitDefaultModesAndTransitions(TObjectPtr<UImmediateMovementModeTransition> InImmediateMovementModeTransition, TObjectPtr<UNullMovementMode> InNullMovementMode)
+	{
+		ImmediateMovementModeTransitionWeakPtr = TWeakObjectPtr<UImmediateMovementModeTransition>(InImmediateMovementModeTransition.Get());
+		NullMovementModeWeakPtr = TWeakObjectPtr<UNullMovementMode>(InNullMovementMode.Get());
+	}
+
+	void FMoverStateMachine::AdvanceToNextMode()
+	{
+		TStrongObjectPtr<UImmediateMovementModeTransition> QueuedModeTransition = QueuedModeTransitionWeakPtr.Pin();
+		if (QueuedModeTransition)
+		{
+			const FName NextModeName = QueuedModeTransition->GetNextModeName();
+
+			if (NextModeName != NAME_None)
+			{
+				TWeakObjectPtr<UBaseMovementMode>* FoundNextMovementMode = Modes.Find(NextModeName);
+				if (FoundNextMovementMode)
+				{
+					const bool bShouldNextModeReenter = QueuedModeTransition->ShouldReenter();
+					if ((CurrentModeName != NextModeName) || bShouldNextModeReenter)
+					{
+						const FName PreviousModeName = CurrentModeName;
+						TWeakObjectPtr<UBaseMovementMode>* FoundPreviousMovementMode = Modes.Find(PreviousModeName);
+
+						CurrentModeName = NextModeName;
+
+						UE_LOG(LogMover, Verbose, TEXT("AdvanceToNextMode: %s (%s) from %s to %s"),
+							*OwnerActorName, *UEnum::GetValueAsString(OwnerActorLocalNetRole), *CurrentModeName.ToString(), *NextModeName.ToString());
+
+						TWeakObjectPtr<UBaseMovementMode> PreviousMovementMode = FoundPreviousMovementMode ? *FoundPreviousMovementMode : TWeakObjectPtr<UBaseMovementMode>();
+						QueuedModeChangeEvents.Add({ PreviousModeName, PreviousMovementMode, CurrentModeName, *FoundNextMovementMode });
+					}
+				}
+			}
+
+			ClearQueuedMode();
+		}
+	}
+
+	const TArray<FMovementModeChangeEvent>& FMoverStateMachine::GetQueuedModeChangeEvents() const
+	{
+		return QueuedModeChangeEvents;
+	}
+
+	void FMoverStateMachine::ClearQueuedModeChangeEvents()
+	{
+		QueuedModeChangeEvents.Empty();
+	}
+
+	void FMoverStateMachine::FlushQueuedMovesToGroup(FLayeredMoveGroup& Group)
+	{
+		if (!QueuedLayeredMoves.IsEmpty())
+		{
+			for (TSharedPtr<FLayeredMoveBase>& QueuedMove : QueuedLayeredMoves)
+			{
+				Group.QueueLayeredMove(QueuedMove);
+			}
+
+			QueuedLayeredMoves.Empty();
+		}
+	}
+
+	bool FMoverStateMachine::ApplyInstantEffects(const FMoverTickStartData& SubstepStartData, const FMoverTimeStep& SubTimeStep, FMoverSyncState& OutputState)
+	{
+		FApplyMovementEffectParams_Async EffectParams_Async;
+		EffectParams_Async.StartState = &SubstepStartData;
+		EffectParams_Async.TimeStep = &SubTimeStep;
+
+		bool bInstantMovementEffectApplied = false;
+
+		if (!QueuedInstantEffects.IsEmpty())
+		{
+			for (TSharedPtr<FInstantMovementEffect>& QueuedEffect : QueuedInstantEffects)
+			{
+				bInstantMovementEffectApplied |= QueuedEffect->ApplyMovementEffect_Async(EffectParams_Async, OutputState);
+			}
+
+			QueuedInstantEffects.Empty();
+		}
+
+		return bInstantMovementEffectApplied;
+	}
+
+	void FMoverStateMachine::SetOwnerActorName(const FString& InOwnerActorName)
+	{
+		OwnerActorName = InOwnerActorName;
+	}
+	
+	void FMoverStateMachine::SetOwnerActorLocalNetRole(ENetRole InOwnerActorLocalNetRole)
+	{
+		OwnerActorLocalNetRole = InOwnerActorLocalNetRole;
+	}
+
+} // End of namespace UE::ChaosMover
\ No newline at end of file
diff --git a/Engine/Plugins/Experimental/ChaosMover/Source/ChaosMover/Private/ChaosMoverStateMachine.h b/Engine/Plugins/Experimental/ChaosMover/Source/ChaosMover/Private/ChaosMoverStateMachine.h
new file mode 100644
--- /dev/null
+++ b/Engine/Plugins/Experimental/ChaosMover/Source/ChaosMover/Private/ChaosMoverStateMachine.h
@@ -0,0 +1,96 @@
+// Copyright Epic Games, Inc. All Rights Reserved.
+
+#pragma once
+
+#include "CoreMinimal.h"
+#include "Engine/EngineTypes.h"
+#include "Templates/SubclassOf.h"
+#include "MovementMode.h"
+#include "InstantMovementEffect.h"
+
+namespace UE::ChaosMover
+{
+	/** Mode change event structure, used to postpone callbacks to gameplay code when a mode has changed */
+	struct FMovementModeChangeEvent
+	{
+		FName PreviousMovementModeName;
+		TWeakObjectPtr<UBaseMovementMode> PreviousMovementMode;
+		FName NextMovementModeName;
+		TWeakObjectPtr<UBaseMovementMode> NextMovementMode;
+	};
+
+	class CHAOSMOVER_API FMoverStateMachine
+	{
+	public:
+		FMoverStateMachine();
+		virtual ~FMoverStateMachine();
+
+		void RegisterMovementMode(FName ModeName, TObjectPtr<UBaseMovementMode> Mode, bool bIsDefaultMode = false);
+
+		void UnregisterMovementMode(FName ModeName);
+		void ClearAllMovementModes();
+
+		void RegisterGlobalTransition(TObjectPtr<UBaseMovementModeTransition> Transition);
+		void UnregisterGlobalTransition(TObjectPtr<UBaseMovementModeTransition> Transition);
+		void ClearAllGlobalTransitions();
+
+		void SetDefaultMode(FName NewDefaultModeName);
+
+		void QueueNextMode(FName DesiredNextModeName, bool bShouldReenter = false);
+		void SetModeImmediately(FName DesiredModeName, bool bShouldReenter = false);
+		void ClearQueuedMode();
+
+		void OnSimulationTick(const FMoverTimeStep& TimeStep, const FMoverTickStartData& StartState, UMoverBlackboard* SimBlackboard, UMovementMixer* MovementMixer, FMoverTickEndData& OutputState);
+		void OnSimulationPreRollback(const FMoverSyncState* InvalidSyncState, const FMoverSyncState* SyncState, const FMoverAuxStateContext* InvalidAuxState, const FMoverAuxStateContext* AuxState);
+		void OnSimulationRollback(const FMoverSyncState* SyncState, const FMoverAuxStateContext* AuxState);
+
+		FName GetCurrentModeName() const { return CurrentModeName; }
+
+		const UBaseMovementMode* GetCurrentMode() const;
+
+		const UBaseMovementMode* FindMovementMode(FName ModeName) const;
+
+		void QueueLayeredMove(TSharedPtr<FLayeredMoveBase> Move);
+
+		void QueueInstantMovementEffect(TSharedPtr<FInstantMovementEffect> Effect);
+
+		const TArray<FMovementModeChangeEvent>& GetQueuedModeChangeEvents() const;
+		void ClearQueuedModeChangeEvents();
+
+		void SetOwnerActorName(const FString& InOwnerActorName);
+		void SetOwnerActorLocalNetRole(ENetRole InOwnerActorLocalNetRole);
+
+		void InitDefaultModesAndTransitions(TObjectPtr<UImmediateMovementModeTransition> ImmediateMovementModeTransition, TObjectPtr<UNullMovementMode> NullMovementMode);
+
+	protected:
+
+		TMap<FName, TWeakObjectPtr<UBaseMovementMode>> Modes;
+		TArray<TWeakObjectPtr<UBaseMovementModeTransition>> GlobalTransitions;
+
+		TWeakObjectPtr<UImmediateMovementModeTransition> QueuedModeTransitionWeakPtr;
+
+		TWeakObjectPtr<UImmediateMovementModeTransition> ImmediateMovementModeTransitionWeakPtr;
+		TWeakObjectPtr<UNullMovementMode> NullMovementModeWeakPtr;
+
+		FString OwnerActorName;
+		ENetRole OwnerActorLocalNetRole;
+
+		FName DefaultModeName = NAME_None;
+		FName CurrentModeName = NAME_None;
+
+		/** Moves that are queued to be added to the simulation at the start of the next sim subtick */
+		TArray<TSharedPtr<FLayeredMoveBase>> QueuedLayeredMoves;
+
+		/** Effects that are queued to be applied to the simulation at the start of the next sim subtick or at the end of this tick */
+		TArray<TSharedPtr<FInstantMovementEffect>> QueuedInstantEffects;
+
+		/** This queue is used to record movement mode transitions and postpone the execution of their callbacks */
+		TArray<FMovementModeChangeEvent> QueuedModeChangeEvents;
+
+	private:
+		void ConstructDefaultModes();
+		void AdvanceToNextMode();
+		void FlushQueuedMovesToGroup(FLayeredMoveGroup& Group);
+		bool ApplyInstantEffects(const FMoverTickStartData& SubstepStartData, const FMoverTimeStep& SubTimeStep, FMoverSyncState& OutputState);
+	};
+}
\ No newline at end of file
diff --git a/Engine/Plugins/Experimental/ChaosMover/Source/ChaosMover/Private/Character/Modes/ChaosCharacterMovementMode.cpp b/Engine/Plugins/Experimental/ChaosMover/Source/ChaosMover/Private/Character/Modes/ChaosCharacterMovementMode.cpp
new file mode 100644
--- /dev/null
+++ b/Engine/Plugins/Experimental/ChaosMover/Source/ChaosMover/Private/Character/Modes/ChaosCharacterMovementMode.cpp
@@ -0,0 +1,92 @@
+// Copyright Epic Games, Inc. All Rights Reserved.
+
+#include "ChaosMover/Character/Modes/ChaosCharacterMovementMode.h"
+
+#include "Chaos/Character/CharacterGroundConstraintSettings.h"
+#include "ChaosMover/Character/Settings/SharedChaosCharacterMovementSettings.h"
+#include "Components/CapsuleComponent.h"
+#include "Components/SkeletalMeshComponent.h"
+#include "GameFramework/Character.h"
+#include "Math/UnitConversion.h"
+
+#include UE_INLINE_GENERATED_CPP_BY_NAME(ChaosCharacterMovementMode)
+
+UChaosCharacterMovementMode::UChaosCharacterMovementMode(const FObjectInitializer& ObjectInitializer)
+	: Super(ObjectInitializer)
+{
+	SharedSettingsClasses.Add(USharedChaosCharacterMovementSettings::StaticClass());
+}
+
+void UChaosCharacterMovementMode::OnRegistered(const FName ModeName)
+{
+	Super::OnRegistered(ModeName);
+
+	if (TargetHeightOverride.IsSet())
+	{
+		TargetHeight = TargetHeightOverride.GetValue();
+	}
+	else if (const ACharacter* CharacterOwner = GetTypedOuter<ACharacter>())
+	{
+		TargetHeight = -CharacterOwner->GetMesh()->GetRelativeLocation().Z;
+	}
+
+	SharedSettings = GetMoverComponent()->FindSharedSettings<USharedChaosCharacterMovementSettings>();
+	ensureMsgf(SharedSettings, TEXT("Failed to find instance of USharedChaosCharacterMovementSettings on %s. Movement may not function properly."), *GetPathNameSafe(this));
+}
+
+void UChaosCharacterMovementMode::OnUnregistered()
+{
+	SharedSettings = nullptr;
+
+	Super::OnUnregistered();
+}
+
+void UChaosCharacterMovementMode::SetTargetHeightOverride(float InTargetHeight)
+{
+	TargetHeightOverride = InTargetHeight;
+	TargetHeight = InTargetHeight;
+}
+
+void UChaosCharacterMovementMode::ClearTargetHeightOverride()
+{
+	TargetHeightOverride.Reset();
+
+	if (const ACharacter* CharacterOwner = GetTypedOuter<ACharacter>())
+	{
+		TargetHeight = -CharacterOwner->GetMesh()->GetRelativeLocation().Z;
+	}
+	else
+	{
+		TargetHeight = GetDefault<UChaosCharacterMovementMode>(GetClass())->TargetHeight;
+	}
+}
+
+void UChaosCharacterMovementMode::SetQueryRadiusOverride(float InQueryRadius)
+{
+	QueryRadiusOverride = InQueryRadius;
+	QueryRadius = InQueryRadius;
+}
+
+void UChaosCharacterMovementMode::ClearQueryRadiusOverride()
+{
+	QueryRadiusOverride.Reset();
+
+	if (const ACharacter* CharacterOwner = GetTypedOuter<ACharacter>())
+	{
+		if (UCapsuleComponent* CapsuleComp = CharacterOwner->GetCapsuleComponent())
+		{
+			QueryRadius = FMath::Max(CapsuleComp->GetScaledCapsuleRadius() - 5.0f, 0.0f);
+			return;
+		}
+	}
+	
+	QueryRadius = GetDefault<UChaosCharacterMovementMode>(GetClass())->QueryRadius;
+}
+
+void UChaosCharacterMovementMode::UpdateConstraintSettings(Chaos::FCharacterGroundConstraintSettings& ConstraintSettings) const
+{
+	ConstraintSettings.RadialForceLimit = FUnitConversion::Convert(RadialForceLimit, EUnit::Newtons, EUnit::KilogramCentimetersPerSecondSquared);
+	ConstraintSettings.TwistTorqueLimit = FUnitConversion::Convert(TwistTorqueLimit, EUnit::NewtonMeters, EUnit::KilogramCentimetersSquaredPerSecondSquared);
+	ConstraintSettings.SwingTorqueLimit = FUnitConversion::Convert(SwingTorqueLimit, EUnit::NewtonMeters, EUnit::KilogramCentimetersSquaredPerSecondSquared);
+	ConstraintSettings.TargetHeight = TargetHeight;
+}
diff --git a/Engine/Plugins/Experimental/ChaosMover/Source/ChaosMover/Private/Character/Modes/ChaosFallingMode.cpp b/Engine/Plugins/Experimental/ChaosMover/Source/ChaosMover/Private/Character/Modes/ChaosFallingMode.cpp
new file mode 100644
--- /dev/null
+++ b/Engine/Plugins/Experimental/ChaosMover/Source/ChaosMover/Private/Character/Modes/ChaosFallingMode.cpp
@@ -0,0 +1,24 @@
+// Copyright Epic Games, Inc. All Rights Reserved.
+
+#include "ChaosMover/Character/Modes/ChaosFallingMode.h"
+
+#include UE_INLINE_GENERATED_CPP_BY_NAME(ChaosFallingMode)
+
+UChaosFallingMode::UChaosFallingMode(const FObjectInitializer& ObjectInitializer)
+	: Super(ObjectInitializer)
+{
+	bSupportsAsync = true;
+
+	GameplayTags.AddTag(Mover_IsInAir);
+	GameplayTags.AddTag(Mover_IsFalling);
+}
+
+void UChaosFallingMode::GenerateMove_Implementation(const FMoverTickStartData& StartState, const FMoverTimeStep& TimeStep, FProposedMove& OutProposedMove) const
+{
+
+}
+
+void UChaosFallingMode::SimulationTick_Implementation(const FSimulationTickParams& Params, FMoverTickEndData& OutputState)
+{
+
+}
\ No newline at end of file
diff --git a/Engine/Plugins/Experimental/ChaosMover/Source/ChaosMover/Private/Character/Modes/ChaosWalkingMode.cpp b/Engine/Plugins/Experimental/ChaosMover/Source/ChaosMover/Private/Character/Modes/ChaosWalkingMode.cpp
new file mode 100644
--- /dev/null
+++ b/Engine/Plugins/Experimental/ChaosMover/Source/ChaosMover/Private/Character/Modes/ChaosWalkingMode.cpp
@@ -0,0 +1,379 @@
+// Copyright Epic Games, Inc. All Rights Reserved.
+
+#include "ChaosMover/Character/Modes/ChaosWalkingMode.h"
+#include "MoveLibrary/MovementUtils.h"
+#include "MoveLibrary/FloorQueryUtils.h"
+#include "MoveLibrary/WaterMovementUtils.h"
+#include "ChaosMover/ChaosMoverSimulationTypes.h"
+#include "MoveLibrary/GroundMovementUtils.h"
+#include "DefaultMovementSet/Settings/CommonLegacyMovementSettings.h"
+
+#include UE_INLINE_GENERATED_CPP_BY_NAME(ChaosWalkingMode)
+
+UChaosWalkingMode::UChaosWalkingMode(const FObjectInitializer& ObjectInitializer)
+	: Super(ObjectInitializer)
+{
+	bSupportsAsync = true;
+	GameplayTags.AddTag(Mover_IsOnGround);
+}
+
+void UChaosWalkingMode::GenerateMove_Implementation(const FMoverTickStartData& StartState, const FMoverTimeStep& TimeStep, FProposedMove& OutProposedMove) const
+{
+
+}
+
+/*
+* WIP
+* 
+*
+
+USTRUCT()
+struct CHAOSMOVER_API FGroundMovementModeData
+{
+	UPROPERTY()
+	float TargetHeight = 95.0f;
+
+	UPROPERTY()
+	float QueryRadius = 30.0f;
+};
+
+
+* 
+void FloorCheck(
+	const FMoverDefaultSyncState& SyncState,
+	const FChaosMoverSimulationDefaultInputs& ChaosMoverInput,
+	const FProposedMove& ProposedMove,
+	float DeltaSeconds,
+	float TargetHeight,
+	FFloorCheckResult& OutFloorResult,
+	FWaterCheckResult& OutWaterResult,
+	FVector& OutDeltaPos);
+
+bool CanStepUpOnHitSurface(
+	const FFloorCheckResult& FloorResult,
+	float TargetHeight,
+	const UCommonLegacyMovementSettings* CommonLegacySettings)
+{
+	const float StepHeight = TargetHeight - FloorResult.FloorDist;
+
+	bool bWalkable = StepHeight <= CommonLegacySettings->MaxStepHeight;
+	constexpr float MinStepHeight = 2.0f;
+	const bool SteppingUp = StepHeight > MinStepHeight;
+	if (bWalkable && SteppingUp)
+	{
+		bWalkable = UGroundMovementUtils::CanStepUpOnHitSurface(FloorResult.HitResult);
+	}
+
+	return bWalkable;
+}
+*/
+
+void UChaosWalkingMode::SimulationTick_Implementation(const FSimulationTickParams& Params, FMoverTickEndData& OutputState)
+{
+	/*
+	* WIP
+	* 
+	const FMoverTickStartData& StartState = Params.StartState;
+	FProposedMove ProposedMove = Params.ProposedMove;
+
+	const FChaosMoverSimulationDefaultInputs* ChaosMoverInfo = StartState.AuxState.AuxStateCollection.FindDataByType<FChaosMoverSimulationDefaultInputs>();
+	const FVector UpDir = ChaosMoverInfo ? UMovementUtils::DeduceUpDirectionFromGravity(ChaosMoverInfo->Gravity) : FVector::UpVector;
+	const FVector GravityAcceleration = ChaosMoverInfo ? ChaosMoverInfo->Gravity : -980 * UpDir;
+
+	const FMoverDefaultSyncState* StartingSyncState = StartState.SyncState.SyncStateCollection.FindDataByType<FMoverDefaultSyncState>();
+	check(StartingSyncState);
+
+	FMoverDefaultSyncState& OutputSyncState = OutputState.SyncState.SyncStateCollection.FindOrAddMutableDataByType<FMoverDefaultSyncState>();
+
+	const float DeltaSeconds = Params.TimeStep.StepMs * 0.001f;
+
+	UMoverBlackboard* SimBlackboard = Params.SimBlackboard;
+	if (!SimBlackboard)
+	{
+		OutputSyncState = *StartingSyncState;
+		return;
+	}
+
+	// Store the previous ground normal that was used to compute the proposed move
+	FFloorCheckResult PrevFloorResult;
+	FVector PrevGroundNormal = UpDir; 
+	if (SimBlackboard->TryGet(CommonBlackboard::LastFloorResult, PrevFloorResult))
+	{
+		PrevGroundNormal = PrevFloorResult.HitResult.ImpactNormal;
+	}
+
+	// Floor query
+	FFloorCheckResult FloorResult;
+	FWaterCheckResult WaterResult;
+	FVector OutDeltaPos;
+
+	FloorCheck(*StartingSyncState, *ChaosMoverInfo, ProposedMove, DeltaSeconds, TargetHeight, FloorResult, WaterResult, OutDeltaPos);
+
+	ProposedMove.LinearVelocity = OutDeltaPos / DeltaSeconds;
+
+	// The base movement mode does not apply gravity in walking mode so apply here
+	// This is so that the gravity in this mode will be consistent with the gravity
+	// set on the mover, not the default physics gravity
+	FVector TargetVelocity = StartingSyncState->GetVelocity_WorldSpace() + UMovementUtils::ComputeVelocityFromGravity(GravityAcceleration, DeltaSeconds);
+
+	SimBlackboard->Set(CommonBlackboard::LastFloorResult, FloorResult);
+	SimBlackboard->Set(CommonBlackboard::LastWaterResult, WaterResult);
+
+	const bool bStartSwimming = WaterResult.WaterSplineData.ImmersionDepth > CommonLegacySettings->SwimmingStartImmersionDepth;
+
+	if (WaterResult.IsSwimmableVolume() && bStartSwimming)
+	{
+		SwitchToState(CommonLegacySettings->SwimmingMovementModeName, Params, OutputState);
+	}
+	else if (FloorResult.IsWalkableFloor())
+	{
+		const FVector StartGroundVelocity = UPhysicsMovementUtils::ComputeGroundVelocityFromHitResult(StartingSyncState->GetLocation_WorldSpace(), FloorResult.HitResult, DeltaSeconds);
+		const FVector ProjectedGroundVelocity = UPhysicsMovementUtils::ComputeIntegratedGroundVelocityFromHitResult(StartingSyncState->GetLocation_WorldSpace(), FloorResult.HitResult, DeltaSeconds);
+		const bool bIsGroundMoving = ProjectedGroundVelocity.SizeSquared() > UE_KINDA_SMALL_NUMBER;
+
+		TargetVelocity = StartingSyncState->GetVelocity_WorldSpace();
+		FVector TargetPosition = StartingSyncState->GetLocation_WorldSpace();
+		if (FloorResult.FloorDist > TargetHeight + UE_KINDA_SMALL_NUMBER)
+		{
+			TargetPosition += UpDir * (TargetHeight - FloorResult.FloorDist);
+		}
+
+		constexpr float ParallelCosThreshold = 0.999f;
+		const bool bVerticalVelocity = FVector::Parallel(TargetVelocity.GetSafeNormal(), UpDir, ParallelCosThreshold);
+		const bool bUseProposedMove = !(bHandleVerticalLandingSeparately && bVerticalVelocity) || ProposedMove.bHasDirIntent;
+
+		if (bUseProposedMove)
+		{
+			const FVector ProposedMovePlaneVelocity = ProposedMove.LinearVelocity - ProposedMove.LinearVelocity.ProjectOnToNormal(PrevGroundNormal);
+
+			// If there is velocity intent in the normal direction then use the velocity from the proposed move. Otherwise
+			// retain the previous vertical velocity
+			FVector ProposedNormalVelocity = ProposedMove.LinearVelocity - ProposedMovePlaneVelocity;
+			if (ProposedNormalVelocity.SizeSquared() > UE_KINDA_SMALL_NUMBER)
+			{
+				TargetVelocity += ProposedNormalVelocity - TargetVelocity.ProjectOnToNormal(PrevGroundNormal);
+			}
+
+			TargetPosition += ProposedMovePlaneVelocity * DeltaSeconds;
+		}
+
+		// Check if the proposed velocity would lift off the movement surface.
+		bool bIsLiftingOffSurface = false;
+
+		float CharacterGravity = 0.0f;
+		if (const APhysicsVolume* PhysVolume = UpdatedComponent->GetPhysicsVolume())
+		{
+			CharacterGravity = PhysVolume->GetGravityZ();
+		}
+
+		FVector ProjectedVelocity = TargetVelocity + CharacterGravity * FVector::UpVector * DeltaSeconds;
+
+		const float ProjectedRelativeVerticalVelocity = FloorResult.HitResult.ImpactNormal.Dot(ProjectedVelocity - ProjectedGroundVelocity);
+		const float VerticalVelocityLimit = 2.0f / DeltaSeconds;
+
+		if ((ProjectedRelativeVerticalVelocity > VerticalVelocityLimit) && bIsGroundMoving && (ProjectedVelocity.Dot(UpDir) > VerticalVelocityLimit))
+		{
+			bIsLiftingOffSurface = true;
+		}
+
+		// Determine if the character is stepping up or stepping down.
+		// If stepping up make sure that the step height is less than the max step height
+		// and the new surface has CanCharacterStepUpOn set to true.
+		// If stepping down make sure the step height is less than the max step height.
+		const float StartHeightAboveGround = FloorResult.FloorDist - TargetHeight;
+		const float EndHeightAboutGround = StartHeightAboveGround + UpDir.Dot(ProjectedVelocity - ProjectedGroundVelocity) * DeltaSeconds ;
+		const bool bIsSteppingDown = StartHeightAboveGround > GPhysicsDrivenMotionDebugParams.MinStepUpDistance;
+		const bool bIsWithinReach = EndHeightAboutGround <= CommonLegacySettings->MaxStepHeight;
+
+		// If the character is unsupported allow some grace period before falling
+		bool bIsSupported = bIsWithinReach && !bIsLiftingOffSurface;
+		float TimeSinceSupported = MaxUnsupportedTimeBeforeFalling;
+		SimBlackboard->TryGet(CommonBlackboard::TimeSinceSupported, TimeSinceSupported);
+		if (bIsSupported)
+		{
+			SimBlackboard->Set(CommonBlackboard::TimeSinceSupported, 0.0f);
+		}
+		else if (!bIsLiftingOffSurface)
+		{
+			// Falling
+			TimeSinceSupported += DeltaSeconds;
+			SimBlackboard->Set(CommonBlackboard::TimeSinceSupported, TimeSinceSupported);
+			bIsSupported = TimeSinceSupported < MaxUnsupportedTimeBeforeFalling;
+		}
+		else
+		{
+			// Moving up relative to ground
+			SimBlackboard->Set(CommonBlackboard::TimeSinceSupported, MaxUnsupportedTimeBeforeFalling);
+		}
+
+		// Apply vertical velocity to target if stepping down
+		const bool bNeedsVerticalVelocityToTarget = bIsSupported && bIsSteppingDown && (EndHeightAboutGround > 0.0f) && !bIsLiftingOffSurface;
+		if (bNeedsVerticalVelocityToTarget)
+		{
+			TargetVelocity -= FractionalDownwardVelocityToTarget * (EndHeightAboutGround / DeltaSeconds) * UpDir;
+		}
+
+		// Target orientation
+		// This is always applied regardless of whether the character is supported
+		FRotator TargetOrientation = StartingSyncState->GetOrientation_WorldSpace();
+		if (!UMovementUtils::IsAngularVelocityZero(ProposedMove.AngularVelocity)) 
+		{
+			TargetOrientation += (ProposedMove.AngularVelocity * DeltaSeconds);
+		}
+
+		if (bIsSupported)
+		{
+			OutputState.MovementEndState.RemainingMs = 0.0f;
+			OutputSyncState.MoveDirectionIntent = ProposedMove.bHasDirIntent ? ProposedMove.DirectionIntent : FVector::ZeroVector;
+			OutputSyncState.SetTransforms_WorldSpace(
+				TargetPosition,
+				TargetOrientation,
+				TargetVelocity,
+				nullptr);
+		}
+		else
+		{
+			// Blocking hit but not supported
+			SwitchToState(CommonLegacySettings->AirMovementModeName, Params, OutputState);
+		}
+	}
+	else
+	{
+		// No water or floor not found
+		float TimeSinceSupported = MaxUnsupportedTimeBeforeFalling;
+		SimBlackboard->TryGet(CommonBlackboard::TimeSinceSupported, TimeSinceSupported);
+		TimeSinceSupported += DeltaSeconds;
+		SimBlackboard->Set(CommonBlackboard::TimeSinceSupported, TimeSinceSupported);
+		if (TimeSinceSupported >= MaxUnsupportedTimeBeforeFalling)
+		{
+			SwitchToState(CommonLegacySettings->AirMovementModeName, Params, OutputState);
+		}
+	}
+	*/
+}
+
+
+/*
+*
+* WIP
+void FloorCheck(const FMoverDefaultSyncState& SyncState,
+	const FChaosMoverSimulationDefaultInputs& ChaosMoverInput,
+	const FProposedMove& ProposedMove,
+	float DeltaSeconds,
+	float TargetHeight,
+	FFloorCheckResult& OutFloorResult,
+	FWaterCheckResult& OutWaterResult,
+	FVector& OutDeltaPos)
+{
+	const FVector& UpDir = ChaosMoverInput.UpDir;
+
+	FVector DeltaPos = ProposedMove.LinearVelocity * DeltaSeconds;
+	OutDeltaPos = DeltaPos;
+
+	const float FloorSweepDistance = TargetHeight + CommonLegacySettings->MaxStepHeight;
+
+	UPhysicsMovementUtils::FloorSweep_Internal(SyncState.GetLocation_WorldSpace(), DeltaPos, UpdatedPrimitive, UpDir,
+		QueryRadius, FloorSweepDistance, CommonLegacySettings->MaxWalkSlopeCosine, TargetHeight, OutFloorResult, OutWaterResult);
+
+	if (!OutFloorResult.bBlockingHit)
+	{
+		// Floor not found
+		return;
+	}
+
+	bool bWalkableFloor = OutFloorResult.bWalkableFloor && CanStepUpOnHitSurface(OutFloorResult, TargetHeight, );
+	if (bWalkableFloor)
+	{
+		// Walkable floor found
+		return;
+	}
+
+	// Hit something but not walkable. Try a new query to find a walkable surface
+	const float StepBlockedHeight = TargetHeight - ChaosMoverInput.BoundingCylinderfHeight + ChaosMoverInput.BoundingCylinderRadius;
+	const float StepHeight = TargetHeight - OutFloorResult.FloorDist;
+
+	if (StepHeight > StepBlockedHeight)
+	{
+		// Collision should prevent movement. Just try to find ground at start of movement
+		const float ShrinkMultiplier = 0.75f;
+		UPhysicsMovementUtils::FloorSweep_Internal(SyncState.GetLocation_WorldSpace(), FVector::ZeroVector, UpdatedPrimitive, UpDir,
+			ShrinkMultiplier * QueryRadius, FloorSweepDistance, CommonLegacySettings->MaxWalkSlopeCosine, TargetHeight, OutFloorResult, OutWaterResult);
+
+		OutFloorResult.bWalkableFloor = OutFloorResult.bWalkableFloor && CanStepUpOnHitSurface(OutFloorResult);
+		return;
+	}
+
+	if (DeltaPos.SizeSquared() < UE_SMALL_NUMBER)
+	{
+		// Stationary
+		OutDeltaPos = FVector::ZeroVector;
+		return;
+	}
+
+	// Try to limit the movement to remain on a walkable surface
+	FVector NewDeltaPos = DeltaPos;
+	float NewQueryRadius = QueryRadius;
+
+	FVector HorizSurfaceDir = FVector::VectorPlaneProject(OutFloorResult.HitResult.ImpactNormal, UpDir);
+	float HorizSurfaceDirSizeSq = HorizSurfaceDir.SizeSquared();
+	bool bFoundOutwardDir = false;
+	if (HorizSurfaceDirSizeSq > UE_SMALL_NUMBER)
+	{
+		HorizSurfaceDir *= FMath::InvSqrt(HorizSurfaceDirSizeSq);
+		bFoundOutwardDir = true;
+	}
+	else
+	{
+		// Flat unwalkable surface. Try and get the horizontal direction from the normal instead
+		HorizSurfaceDir = FVector::VectorPlaneProject(OutFloorResult.HitResult.Normal, UpDir);
+		HorizSurfaceDirSizeSq = HorizSurfaceDir.SizeSquared();
+
+		if (HorizSurfaceDirSizeSq > UE_SMALL_NUMBER)
+		{
+			HorizSurfaceDir *= FMath::InvSqrt(HorizSurfaceDirSizeSq);
+			bFoundOutwardDir = true;
+		}
+	}
+
+	if (bFoundOutwardDir)
+	{
+		// If we're moving away try a ray query at the end of the motion
+		const float DP = DeltaPos.Dot(HorizSurfaceDir);
+		if (DP > 0.0f)
+		{
+			NewQueryRadius = 0.0f;
+		}
+		else
+		{
+			NewQueryRadius = 0.25f * QueryRadius;
+			NewDeltaPos = DeltaPos - DP * HorizSurfaceDir;
+		}
+
+		UPhysicsMovementUtils::FloorSweep_Internal(SyncState.GetLocation_WorldSpace(), NewDeltaPos, UpdatedPrimitive, UpDir,
+			NewQueryRadius, FloorSweepDistance, CommonLegacySettings->MaxWalkSlopeCosine, TargetHeight, OutFloorResult, OutWaterResult);
+
+		OutFloorResult.bWalkableFloor = OutFloorResult.bWalkableFloor && CanStepUpOnHitSurface(OutFloorResult);
+		if (OutFloorResult.bWalkableFloor)
+		{
+			OutDeltaPos = NewDeltaPos;
+		}
+		else
+		{
+			OutFloorResult.bWalkableFloor = false;
+		}
+	}
+	else
+	{
+		// Try a query at the start of the movement to find a walkable surface and prevent movement
+
+		NewDeltaPos = FVector::ZeroVector;
+		NewQueryRadius = 0.25f * QueryRadius;
+
+		UPhysicsMovementUtils::FloorSweep_Internal(SyncState.GetLocation_WorldSpace(), NewDeltaPos, UpdatedPrimitive, UpDir,
+			NewQueryRadius, FloorSweepDistance, CommonLegacySettings->MaxWalkSlopeCosine, TargetHeight, OutFloorResult, OutWaterResult);
+
+		OutFloorResult.bWalkableFloor = OutFloorResult.bWalkableFloor && CanStepUpOnHitSurface(OutFloorResult);
+		OutDeltaPos = NewDeltaPos;
+	}
+}
+*/
diff --git a/Engine/Plugins/Experimental/ChaosMover/Source/ChaosMover/Private/Utilities/ChaosMoverQueryUtils.cpp b/Engine/Plugins/Experimental/ChaosMover/Source/ChaosMover/Private/Utilities/ChaosMoverQueryUtils.cpp
new file mode 100644
--- /dev/null
+++ b/Engine/Plugins/Experimental/ChaosMover/Source/ChaosMover/Private/Utilities/ChaosMoverQueryUtils.cpp
@@ -0,0 +1,145 @@
+// Copyright Epic Games, Inc. All Rights Reserved.
+
+#include "ChaosMover/Utilities/ChaosMoverQueryUtils.h"
+
+#include "Chaos/PhysicsObject.h"
+#include "Chaos/PhysicsObjectInternalInterface.h"
+#include "ChaosMover/ChaosMoverDeveloperSettings.h"
+#include "Framework/Threading.h"
+#include "MoveLibrary/FloorQueryUtils.h"
+#include "MoveLibrary/WaterMovementUtils.h"
+#include "Physics/GenericPhysicsInterface.h"
+#include "WaterBodyActor.h"
+
+#if CHAOSMOVER_DEBUG_DRAW
+#include "Chaos/DebugDrawQueue.h"
+#endif
+
+namespace UE::ChaosMover::Utils
+{
+	namespace Internal
+	{
+		bool GetWaterResultFromHitResults_Internal(const TArray<FHitResult>& Hits, const FVector& Location, const float TargetHeight, FWaterCheckResult& OutWaterResult)
+		{
+			// Find the closet hit that is a water body
+			// Note: Relies on ordering of hit results
+			for (int32 Idx = 0; Idx < Hits.Num(); ++Idx)
+			{
+				const FHitResult& Hit = Hits[Idx];
+
+				if (Hit.Component.IsValid())
+				{
+					if (AActor* Actor = Hit.Component->GetOwner())
+					{
+						if (Actor->IsA(AWaterBody::StaticClass()))
+						{
+							AWaterBody* WaterBody = Cast<AWaterBody>(Actor);
+
+							OutWaterResult.HitResult = Hit;
+							OutWaterResult.bSwimmableVolume = true;
+
+							OutWaterResult.WaterSplineData.SplineInputKey = WaterBody->GetWaterBodyComponent()->FindInputKeyClosestToWorldLocation(Location);
+
+							OutWaterResult.WaterSplineData.WaterBody = WaterBody;
+
+							FWaterBodyQueryResult QueryResult = WaterBody->GetWaterBodyComponent()->QueryWaterInfoClosestToWorldLocation(
+								Location,
+								EWaterBodyQueryFlags::ComputeLocation | EWaterBodyQueryFlags::ComputeNormal | EWaterBodyQueryFlags::ComputeImmersionDepth,
+								OutWaterResult.WaterSplineData.SplineInputKey
+							);
+
+							OutWaterResult.WaterSplineData.ImmersionDepth = QueryResult.GetImmersionDepth();
+
+							OutWaterResult.WaterSplineData.WaterPlaneLocation = QueryResult.GetWaterPlaneLocation();
+							OutWaterResult.WaterSplineData.WaterPlaneNormal = QueryResult.GetWaterPlaneNormal();
+
+							const float CapsuleBottom = Location.Z;
+							const float CapsuleTop = Location.Z + (TargetHeight * 2);
+							OutWaterResult.WaterSplineData.WaterSurfaceLocation = QueryResult.GetWaterSurfaceLocation();
+
+							OutWaterResult.WaterSplineData.WaterSurfaceOffset = OutWaterResult.WaterSplineData.WaterSurfaceLocation - Location;
+
+							OutWaterResult.WaterSplineData.ImmersionPercent = FMath::Clamp((OutWaterResult.WaterSplineData.WaterSurfaceLocation.Z - CapsuleBottom) / (CapsuleTop - CapsuleBottom), 0.f, 1.f);
+
+							OutWaterResult.WaterSplineData.WaterSurfaceNormal = QueryResult.GetWaterSurfaceNormal();
+
+	#if CHAOSMOVER_DEBUG_DRAW
+							Chaos::FDebugDrawQueue::GetInstance().DrawDebugLine(Location, Location - FVector::UpVector * OutWaterResult.WaterSplineData.ImmersionDepth, FColor::Blue, false, -1.f, 10, 1.0f);
+							Chaos::FDebugDrawQueue::GetInstance().DrawDebugPoint(OutWaterResult.HitResult.Location, FColor::Blue, false, -1.f, 10, 1.0f);
+	#endif
+							return true;
+						}
+					}
+				}
+			}
+
+			return false;
+		}
+
+	} // namespace Internal
+
+	void FloorSweep_Internal(const FFloorSweepParams& Params, FFloorCheckResult& OutFloorResult, FWaterCheckResult& OutWaterResult)
+	{
+		Chaos::EnsureIsInPhysicsThreadContext();
+
+		TArray<FHitResult> Hits;
+
+		const float DeltaPosVertLength = Params.DeltaPos.Dot(Params.UpDir);
+		const FVector DeltaPosHoriz = Params.DeltaPos - DeltaPosVertLength * Params.UpDir;
+
+		// Make sure the query is long enough to include the vertical movement
+		const float AdjustedQueryDistance = FMath::Max(UE_KINDA_SMALL_NUMBER + DeltaPosVertLength + Params.TargetHeight, Params.QueryDistance);
+
+		// The bottom of the query shape should be at the integrated location (ignoring vertical movement)
+		FVector Start = Params.Location + DeltaPosHoriz + (Params.QueryRadius + UE_KINDA_SMALL_NUMBER) * Params.UpDir;
+		FVector End = Start - AdjustedQueryDistance * Params.UpDir;
+		FHitResult OutHit;
+		if (Chaos::Private::FGenericPhysicsInterface_Internal::SpherecastMulti(Params.World, Params.QueryRadius, Hits, Start, End, Params.CollisionChannel, Params.QueryParams, Params.ResponseParams))
+		{
+			OutHit = Hits.Last();
+		}
+
+#if CHAOSMOVER_DEBUG_DRAW
+		// Draw full length of query
+		if (CVars::bDrawGroundQueries)
+		{
+			const FVector Center = 0.5f * (Start + End);
+			const float Dist = (Start - End).Size();
+			Chaos::FDebugDrawQueue::GetInstance().DrawDebugCapsule(Center, 0.5f * Dist + Params.QueryRadius, Params.QueryRadius, FQuat::Identity, FColor::Silver, false, -1.f, 10, 1.0f);
+		}
+#endif
+		
+		if (OutHit.bBlockingHit)
+		{
+			bool bWalkable = UFloorQueryUtils::IsHitSurfaceWalkable(OutHit, Params.UpDir, Params.MaxWalkSlopeCosine);
+
+#if CHAOSMOVER_DEBUG_DRAW
+			if (CVars::bDrawGroundQueries)
+			{
+				const FVector Center = Start - 0.5f * OutHit.Distance * Params.UpDir;
+				const FColor Color = bWalkable ? FColor::Green : FColor::Red;
+				Chaos::FDebugDrawQueue::GetInstance().DrawDebugCapsule(Center, 0.5f * OutHit.Distance + Params.QueryRadius, Params.QueryRadius, FQuat::Identity, Color, false, -1.f, 10, 1.0f);
+			}
+#endif
+			OutFloorResult.bBlockingHit = true;
+			OutFloorResult.bWalkableFloor = bWalkable;
+			OutFloorResult.FloorDist = Params.UpDir.Dot(Params.Location - OutHit.ImpactPoint);
+			OutFloorResult.HitResult = OutHit;
+
+			// Make sure that the object stored in the hit result is the root object
+			if (const Chaos::FPhysicsObjectHandle HitObject = OutHit.PhysicsObject)
+			{
+				Chaos::FReadPhysicsObjectInterface_Internal Interface = Chaos::FPhysicsObjectInternalInterface::GetRead();
+				OutFloorResult.HitResult.PhysicsObject = Interface.GetRootObject(HitObject);
+			}
+		}
+		else
+		{
+			OutFloorResult.Clear();
+			OutFloorResult.FloorDist = 1.0e10f;
+		}
+
+		Internal::GetWaterResultFromHitResults_Internal(Hits, Params.Location, Params.TargetHeight, OutWaterResult);
+	}
+
+} // namespace UE::ChaosMover::Utils
\ No newline at end of file
diff --git a/Engine/Plugins/Experimental/ChaosMover/Source/ChaosMover/Public/ChaosMover/Backends/ChaosMoverBackend.h b/Engine/Plugins/Experimental/ChaosMover/Source/ChaosMover/Public/ChaosMover/Backends/ChaosMoverBackend.h
new file mode 100644
--- /dev/null
+++ b/Engine/Plugins/Experimental/ChaosMover/Source/ChaosMover/Public/ChaosMover/Backends/ChaosMoverBackend.h
@@ -0,0 +1,89 @@
+// Copyright Epic Games, Inc. All Rights Reserved.
+
+#pragma once
+
+#include "Backends/MoverBackendLiaison.h"
+#include "Chaos/Declares.h"
+#include "ChaosMover/ChaosMoverSimulationTypes.h"
+#include "Components/ActorComponent.h"
+
+#include "ChaosMoverBackend.generated.h"
+
+namespace UE::ChaosMover
+{
+	class FSimulation;
+}
+
+class UMoverComponent;
+class UNetworkPhysicsComponent;
+
+UCLASS(Within = MoverComponent)
+class CHAOSMOVER_API UChaosMoverBackendComponent : public UActorComponent, public IMoverBackendLiaisonInterface
+{
+	GENERATED_BODY()
+
+public:
+	UChaosMoverBackendComponent();
+
+	virtual void InitializeComponent() override;
+	virtual void UninitializeComponent() override;
+	virtual void BeginPlay() override;
+	virtual void EndPlay(const EEndPlayReason::Type EndPlayReason) override;
+
+	virtual bool IsAsync() const override { return true; }
+
+	// Unused IMoverBackendLiaisonInterface API
+	virtual float GetCurrentSimTimeMs()
+	{
+		return 0.0f;
+	}
+	virtual int32 GetCurrentSimFrame()
+	{
+		return 0;
+	}
+
+	virtual void ProduceInputData(int32 PhysicsStep, int32 NumSteps, double DeltaTimeInSeconds, int32 NetworkPhysicsTickOffset, UE::ChaosMover::FSimulationInputData& InputData);
+	virtual void ConsumeOutputData(const FMoverTimeStep& TimeStep, const UE::ChaosMover::FSimulationOutputData& OutputData);
+	virtual void FinalizeFrame(double ResultsTimeInSeconds);
+
+	void ApplyNetInputData(const FMoverInputCmdContext& InNetInputCmd);
+	void BuildNetInputData(FMoverInputCmdContext& OutNetInputCmd) const;
+	void ApplyNetStateData(const FMoverSyncState& InNetSyncState);
+	void BuildNetStateData(FMoverSyncState& OutNetSyncState) const;
+	void RegisterSimulation() const;
+
+protected:
+	virtual TSharedPtr<UE::ChaosMover::FSimulation> CreateSimulation();
+	virtual void InitSimulation();
+	void InitSimulationStateMachine();
+	virtual void GenerateInput(const FMoverTimeStep& TimeStep, UE::ChaosMover::FSimulationInputData& InputData);
+	/**
+	 * Flushes the queue of movement mode change events, which causes all movement mode change callbacks to be called
+	 */
+	void FlushModeChangeEvents();
+
+	UMoverComponent& GetMoverComponent() const;
+	Chaos::FPhysicsSolver* GetPhysicsSolver() const;
+	Chaos::FPBDRigidParticle* GetControlledParticle() const;
+	FMoverTimeStep GetMoverTimeStep(double DeltaTimeInSeconds, int32 NetworkPhysicsTickOffset) const;
+
+	UFUNCTION()
+	virtual void HandleUpdatedComponentPhysicsStateChanged(UPrimitiveComponent* ChangedComponent, EComponentPhysicsStateChange StateChange);
+
+	UFUNCTION()
+	void HandleOwningPawnControllerChanged_Server(APawn* OwnerPawn, AController* OldController, AController* NewController);
+
+	UPROPERTY(Transient)
+	TObjectPtr<UNetworkPhysicsComponent> NetworkPhysicsComponent;
+
+	TSharedPtr<UE::ChaosMover::FSimulation> Simulation;
+	UE::Mover::FSimulationOutputRecord SimOutputRecord;
+	bool bIsUsingAsyncPhysics = false;
+	bool bIsInputProducer = true;
+
+	UPROPERTY(Transient)
+	TObjectPtr<UMoverBlackboard> MoverBlackboard_Async;
+
+	UPROPERTY(Transient)
+	TObjectPtr<UMovementMixer> MovementMixer_Async;
+};
\ No newline at end of file
diff --git a/Engine/Plugins/Experimental/ChaosMover/Source/ChaosMover/Public/ChaosMover/ChaosMoverDeveloperSettings.h b/Engine/Plugins/Experimental/ChaosMover/Source/ChaosMover/Public/ChaosMover/ChaosMoverDeveloperSettings.h
new file mode 100644
--- /dev/null
+++ b/Engine/Plugins/Experimental/ChaosMover/Source/ChaosMover/Public/ChaosMover/ChaosMoverDeveloperSettings.h
@@ -0,0 +1,19 @@
+// Copyright Epic Games, Inc. All Rights Reserved.
+
+#pragma once
+
+#include "CoreMinimal.h"
+
+namespace UE::ChaosMover
+{
+	namespace CVars
+	{
+		CHAOSMOVER_API extern bool bForceSingleThreadedGT;
+		CHAOSMOVER_API extern bool bForceSingleThreadedPT;
+		CHAOSMOVER_API extern bool bDrawGroundQueries;
+	}
+};
+
+#ifndef CHAOSMOVER_DEBUG_DRAW
+#define CHAOSMOVER_DEBUG_DRAW (!(UE_BUILD_SHIPPING || UE_BUILD_TEST) || WITH_EDITOR)
+#endif
\ No newline at end of file
diff --git a/Engine/Plugins/Experimental/ChaosMover/Source/ChaosMover/Public/ChaosMover/ChaosMoverLog.h b/Engine/Plugins/Experimental/ChaosMover/Source/ChaosMover/Public/ChaosMover/ChaosMoverLog.h
new file mode 100644
--- /dev/null
+++ b/Engine/Plugins/Experimental/ChaosMover/Source/ChaosMover/Public/ChaosMover/ChaosMoverLog.h
@@ -0,0 +1,7 @@
+// Copyright Epic Games, Inc. All Rights Reserved.
+
+#pragma once
+
+#include "Logging/LogMacros.h"
+
+CHAOSMOVER_API DECLARE_LOG_CATEGORY_EXTERN(LogChaosMover, Log, All);
\ No newline at end of file
diff --git a/Engine/Plugins/Experimental/ChaosMover/Source/ChaosMover/Public/ChaosMover/ChaosMoverSimulation.h b/Engine/Plugins/Experimental/ChaosMover/Source/ChaosMover/Public/ChaosMover/ChaosMoverSimulation.h
new file mode 100644
--- /dev/null
+++ b/Engine/Plugins/Experimental/ChaosMover/Source/ChaosMover/Public/ChaosMover/ChaosMoverSimulation.h
@@ -0,0 +1,55 @@
+// Copyright Epic Games, Inc. All Rights Reserved.
+
+#pragma once
+
+#include "ChaosMover/ChaosMoverSimulationTypes.h"
+#include "UObject/WeakObjectPtrTemplates.h"
+#include "ChaosMoverStateMachine.h"
+#include "MoveLibrary/MovementMixer.h"
+#include "MoveLibrary/MoverBlackboard.h"
+
+namespace Chaos
+{
+	class FCollisionContactModifier;
+}
+
+namespace UE::ChaosMover
+{
+	class CHAOSMOVER_API FSimulation
+	{
+	public:
+		FSimulation() = default;
+		FSimulation(const FSimulation&) = delete;
+		const FSimulation& operator=(const FSimulation&) = delete;
+		virtual ~FSimulation();
+
+		void Init(const FMoverSyncState& InitialSyncState);
+		void SimulationTick(const FMoverTimeStep& TimeStep, const FSimulationInputData& InputData, FSimulationOutputData& OutputData);
+		void ModifyContacts(Chaos::FCollisionContactModifier& Modifier);
+
+		void InitNetInputData(const FMoverInputCmdContext& InNetInputCmd);
+		void ApplyNetInputData(const FMoverInputCmdContext& InNetInputCmd);
+		void BuildNetInputData(FMoverInputCmdContext& OutNetInputCmd) const;
+		void ApplyNetStateData(const FMoverSyncState& InNetSyncState);
+		void BuildNetStateData(FMoverSyncState& OutNetSyncState) const;
+
+	protected:
+		virtual void OnInit();
+		virtual void OnPreSimulationTick(const FMoverTimeStep& TimeStep, const FSimulationInputData& InputData);
+		virtual void OnSimulationTick(const FMoverTimeStep& TimeStep, const FSimulationInputData& InputData, FSimulationOutputData& OutputData);
+		virtual void OnPostSimulationTick(const FMoverTimeStep& TimeStep, const FSimulationOutputData& OutputData);
+		virtual void OnModifyContacts(Chaos::FCollisionContactModifier& Modifier);
+
+		FMoverSyncState SyncState;
+
+	public:
+		FMoverStateMachine MoverStateMachine;
+		TWeakObjectPtr<UMoverBlackboard> SimBlackboardWeakPtr;
+		TWeakObjectPtr<UMovementMixer> MovementMixerWeakPtr;
+
+	private:
+		FMoverInputCmdContext InputCmd;
+		bool bInputCmdOverridden = false;
+	};
+
+} // namespace UE::ChaosMover
\ No newline at end of file
diff --git a/Engine/Plugins/Experimental/ChaosMover/Source/ChaosMover/Public/ChaosMover/ChaosMoverSimulationTypes.h b/Engine/Plugins/Experimental/ChaosMover/Source/ChaosMover/Public/ChaosMover/ChaosMoverSimulationTypes.h
new file mode 100644
--- /dev/null
+++ b/Engine/Plugins/Experimental/ChaosMover/Source/ChaosMover/Public/ChaosMover/ChaosMoverSimulationTypes.h
@@ -0,0 +1,63 @@
+// Copyright Epic Games, Inc. All Rights Reserved.
+
+#pragma once
+
+#include "Chaos/Declares.h"
+#include "CollisionQueryParams.h"
+#include "MoverSimulationTypes.h"
+
+#include "ChaosMoverSimulationTypes.generated.h"
+
+namespace UE::ChaosMover
+{
+	struct FSimulationInputData
+	{
+		void Reset()
+		{
+			InputCmd.Reset();
+			InputParams.Empty();
+		}
+
+		mutable FMoverInputCmdContext InputCmd; // can be overridden by network physics
+		FMoverDataCollection InputParams;
+	};
+
+	using FSimulationOutputData = ::UE::Mover::FSimulationOutputData;
+
+} // namespace UE::ChaosMover
+
+// Data block containing all default inputs required by the Chaos Mover simulation
+USTRUCT()
+struct CHAOSMOVER_API FChaosMoverSimulationDefaultInputs : public FMoverDataStructBase
+{
+	GENERATED_BODY()
+
+	FChaosMoverSimulationDefaultInputs()
+	{
+		Reset();
+	}
+
+	void Reset();
+
+	/** return newly allocated copy of this FMoverDataStructBase. MUST be overridden by derived types. */
+	virtual FMoverDataStructBase* Clone() const override
+	{
+		return new FChaosMoverSimulationDefaultInputs();
+	}
+
+	/** Gets the type info of this FMoverDataStructBase. MUST be overridden by derived types. */
+	virtual UScriptStruct* GetScriptStruct() const override
+	{
+		return StaticStruct();
+	}
+
+	FCollisionResponseParams CollisionResponseParams;
+	FCollisionQueryParams CollisionQueryParams;
+	FVector UpDir;
+	FVector Gravity;
+	Chaos::FPhysicsObjectHandle PhysicsObject;
+	AActor* OwningActor;
+	ECollisionChannel CollisionChannel;
+	float BoundingCylinderRadius;
+	float BoundingCylinderfHeight;
+};
diff --git a/Engine/Plugins/Experimental/ChaosMover/Source/ChaosMover/Public/ChaosMover/Character/Modes/ChaosCharacterMovementMode.h b/Engine/Plugins/Experimental/ChaosMover/Source/ChaosMover/Public/ChaosMover/Character/Modes/ChaosCharacterMovementMode.h
new file mode 100644
--- /dev/null
+++ b/Engine/Plugins/Experimental/ChaosMover/Source/ChaosMover/Public/ChaosMover/Character/Modes/ChaosCharacterMovementMode.h
@@ -0,0 +1,85 @@
+// Copyright Epic Games, Inc. All Rights Reserved.
+
+#pragma once
+
+#include "MovementMode.h"
+
+#include "ChaosCharacterMovementMode.generated.h"
+
+namespace Chaos
+{
+	class FCharacterGroundConstraintSettings;
+}
+
+/**
+ * Base class for all Chaos character movement modes
+ */
+UCLASS(Blueprintable, BlueprintType, EditInlineNew, DefaultToInstanced)
+class CHAOSMOVER_API UChaosCharacterMovementMode : public UBaseMovementMode
+{
+	GENERATED_BODY()
+
+public:
+	UChaosCharacterMovementMode(const FObjectInitializer& ObjectInitializer);
+
+	virtual void OnRegistered(const FName ModeName) override;
+	virtual void OnUnregistered() override;
+
+	// Update the constraint settings with the data defined by this mode
+	virtual void UpdateConstraintSettings(Chaos::FCharacterGroundConstraintSettings& ConstraintSettings) const;
+
+	float GetTargetHeight() const
+	{
+		return TargetHeight;
+	}
+
+	void SetTargetHeightOverride(float InTargetHeight);
+	void ClearTargetHeightOverride();
+
+	float GetQueryRadius() const
+	{
+		return QueryRadius;
+	}
+
+	void SetQueryRadiusOverride(float InQueryRadius);
+	void ClearQueryRadiusOverride();
+
+	bool ShouldCharacterRemainUpright() const
+	{
+		return bShouldCharacterRemainUpright;
+	}
+
+	TObjectPtr<const class USharedChaosCharacterMovementSettings> SharedSettings;
+
+protected:
+	// Maximum force the character can apply to reach the motion target
+	UPROPERTY(EditAnywhere, Category = "Constraint Settings", meta = (ClampMin = "0", UIMin = "0", ForceUnits = "Newtons"))
+	float RadialForceLimit = 1500.0f;
+
+	// Maximum torque the character can apply to rotate in air about the vertical axis
+	UPROPERTY(EditAnywhere, Category = "Constraint Settings", meta = (ClampMin = "0", UIMin = "0", ForceUnits = "NewtonMeters"))
+	float TwistTorqueLimit = 1000.0f;
+
+	// Maximum torque the character can apply to remain upright
+	UPROPERTY(EditAnywhere, Category = "Constraint Settings", meta = (ClampMin = "0", UIMin = "0", ForceUnits = "NewtonMeters"))
+	float SwingTorqueLimit = 3000.0f;
+
+	// Controls whether the character capsule is forced to remain upright
+	UPROPERTY(EditAnywhere, Category = "Constraint Settings")
+	bool bShouldCharacterRemainUpright = true;
+
+	/**
+	 * Optional override target height for the character (the desired distance from the center of the capsule to the floor).
+	 * If left blank, the -Z offset of the owning character's skeletal mesh comp will be used automatically.
+	 */
+	UPROPERTY(EditAnywhere, Category = "Constraint Settings", meta = (ClampMin = "0", UIMin = "0", ForceUnits = "cm"))
+	TOptional<float> TargetHeightOverride;
+
+	// Radius used for ground queries
+	UPROPERTY(EditAnywhere, Category = "Query Settings", meta = (ClampMin = "0", UIMin = "0", ForceUnits = "cm"))
+	TOptional<float> QueryRadiusOverride;
+
+private:
+	float TargetHeight = 95.0f;
+	float QueryRadius = 30.0f;
+};
\ No newline at end of file
diff --git a/Engine/Plugins/Experimental/ChaosMover/Source/ChaosMover/Public/ChaosMover/Character/Modes/ChaosFallingMode.h b/Engine/Plugins/Experimental/ChaosMover/Source/ChaosMover/Public/ChaosMover/Character/Modes/ChaosFallingMode.h
new file mode 100644
--- /dev/null
+++ b/Engine/Plugins/Experimental/ChaosMover/Source/ChaosMover/Public/ChaosMover/Character/Modes/ChaosFallingMode.h
@@ -0,0 +1,22 @@
+// Copyright Epic Games, Inc. All Rights Reserved.
+
+#pragma once
+
+#include "ChaosMover/Character/Modes/ChaosCharacterMovementMode.h"
+
+#include "ChaosFallingMode.generated.h"
+
+/**
+ * Chaos character falling mode
+ */
+UCLASS(Blueprintable, BlueprintType, EditInlineNew, DefaultToInstanced)
+class CHAOSMOVER_API UChaosFallingMode : public UChaosCharacterMovementMode
+{
+	GENERATED_BODY()
+
+public:
+	UChaosFallingMode(const FObjectInitializer& ObjectInitializer);
+
+	virtual void GenerateMove_Implementation(const FMoverTickStartData& StartState, const FMoverTimeStep& TimeStep, FProposedMove& OutProposedMove) const override;
+	virtual void SimulationTick_Implementation(const FSimulationTickParams& Params, FMoverTickEndData& OutputState) override;
+};
\ No newline at end of file
diff --git a/Engine/Plugins/Experimental/ChaosMover/Source/ChaosMover/Public/ChaosMover/Character/Modes/ChaosWalkingMode.h b/Engine/Plugins/Experimental/ChaosMover/Source/ChaosMover/Public/ChaosMover/Character/Modes/ChaosWalkingMode.h
new file mode 100644
--- /dev/null
+++ b/Engine/Plugins/Experimental/ChaosMover/Source/ChaosMover/Public/ChaosMover/Character/Modes/ChaosWalkingMode.h
@@ -0,0 +1,22 @@
+// Copyright Epic Games, Inc. All Rights Reserved.
+
+#pragma once
+
+#include "ChaosMover/Character/Modes/ChaosCharacterMovementMode.h"
+
+#include "ChaosWalkingMode.generated.h"
+
+/**
+ * Chaos character walking mode
+ */
+UCLASS(Blueprintable, BlueprintType, EditInlineNew, DefaultToInstanced)
+class CHAOSMOVER_API UChaosWalkingMode : public UChaosCharacterMovementMode
+{
+	GENERATED_BODY()
+
+public:
+	UChaosWalkingMode(const FObjectInitializer& ObjectInitializer);
+
+	virtual void GenerateMove_Implementation(const FMoverTickStartData& StartState, const FMoverTimeStep& TimeStep, FProposedMove& OutProposedMove) const override;
+	virtual void SimulationTick_Implementation(const FSimulationTickParams& Params, FMoverTickEndData& OutputState) override;
+};
\ No newline at end of file
diff --git a/Engine/Plugins/Experimental/ChaosMover/Source/ChaosMover/Public/ChaosMover/Character/Settings/SharedChaosCharacterMovementSettings.h b/Engine/Plugins/Experimental/ChaosMover/Source/ChaosMover/Public/ChaosMover/Character/Settings/SharedChaosCharacterMovementSettings.h
new file mode 100644
--- /dev/null
+++ b/Engine/Plugins/Experimental/ChaosMover/Source/ChaosMover/Public/ChaosMover/Character/Settings/SharedChaosCharacterMovementSettings.h
@@ -0,0 +1,39 @@
+// Copyright Epic Games, Inc. All Rights Reserved.
+
+#pragma once
+
+#include "MovementMode.h"
+#include "SharedChaosCharacterMovementSettings.generated.h"
+
+/**
+ * SharedChaosCharacterMovementSettings: collection of settings that are shared between the chaos character movement modes
+ */
+UCLASS(BlueprintType)
+class CHAOSMOVER_API USharedChaosCharacterMovementSettings : public UObject, public IMovementSettingsInterface
+{
+	GENERATED_BODY()
+
+	virtual FString GetDisplayName() const override
+	{
+		return GetName();
+	}
+
+public:
+	float GetMaxWalkableSlopeCosine() const
+	{
+		return MaxWalkSlopeCosine;
+	}
+
+protected:
+	/** Walkable slope angle, represented as cosine(max slope angle) for performance reasons. Ex: for max slope angle of 30 degrees, value is cosine(30 deg) = 0.866 */
+	UPROPERTY(BlueprintReadWrite, EditAnywhere, Category = "Ground Movement", BlueprintSetter=SetMaxWalkableSlopeAngle, meta = (ClampMin = "0", UIMin = "0", ClampMax = "90", UIMax = "90", ForceUnits = "degrees"))
+	float MaxWalkableSlopeAngle = 45.0f;
+
+	UFUNCTION(BlueprintSetter)
+	void SetMaxWalkableSlopeAngle(float Value)
+	{
+		MaxWalkSlopeCosine = FMath::Cos(FMath::DegreesToRadians(MaxWalkableSlopeAngle));
+	}
+
+	float MaxWalkSlopeCosine = 0.71f;
+};
\ No newline at end of file
diff --git a/Engine/Plugins/Experimental/ChaosMover/Source/ChaosMover/Public/ChaosMover/Utilities/ChaosMoverQueryUtils.h b/Engine/Plugins/Experimental/ChaosMover/Source/ChaosMover/Public/ChaosMover/Utilities/ChaosMoverQueryUtils.h
new file mode 100644
--- /dev/null
+++ b/Engine/Plugins/Experimental/ChaosMover/Source/ChaosMover/Public/ChaosMover/Utilities/ChaosMoverQueryUtils.h
@@ -0,0 +1,30 @@
+// Copyright Epic Games, Inc. All Rights Reserved.
+
+#pragma once
+
+#include "CollisionQueryParams.h"
+
+struct FFloorCheckResult;
+struct FWaterCheckResult;
+class UWorld;
+
+namespace UE::ChaosMover::Utils
+{
+	struct FFloorSweepParams
+	{
+		FCollisionResponseParams ResponseParams;
+		FCollisionQueryParams QueryParams;
+		FVector Location;
+		FVector DeltaPos;
+		FVector UpDir;
+		const UWorld* World;
+		float QueryDistance;
+		float QueryRadius;
+		float MaxWalkSlopeCosine;
+		float TargetHeight;
+		ECollisionChannel CollisionChannel;
+	};
+
+	extern CHAOSMOVER_API void FloorSweep_Internal(const FFloorSweepParams& Params, FFloorCheckResult& OutFloorResult, FWaterCheckResult& OutWaterResult);
+
+}
\ No newline at end of file
diff --git a/Engine/Plugins/Experimental/Mover/Source/Mover/Private/Backends/MoverNetworkPhysicsLiaison.cpp b/Engine/Plugins/Experimental/Mover/Source/Mover/Private/Backends/MoverNetworkPhysicsLiaison.cpp
--- a/Engine/Plugins/Experimental/Mover/Source/Mover/Private/Backends/MoverNetworkPhysicsLiaison.cpp
+++ b/Engine/Plugins/Experimental/Mover/Source/Mover/Private/Backends/MoverNetworkPhysicsLiaison.cpp
@@ -316,6 +316,7 @@ bool UMoverNetworkPhysicsLiaisonComponent::CanSimulate_Internal(const FPhysicsMo
 
 		if (!Cast<IPhysicsCharacterMovementModeInterface>(GetMoverComponent().MovementModes.FindChecked(Input.SyncState.MovementMode)))
 		{
+			ensureMsgf(false, TEXT("Attempting to run non-character physics movement mode %s in physics mover update. Only modes that implement IPhysicsCharacterMovementModeInterface can be used with the CharacterPhysicsLiaison."), *Input.SyncState.MovementMode.ToString());
 			UE_LOG(LogMover, Verbose, TEXT("Attempting to run non-character physics movement mode %s in physics mover update. Only modes that implement IPhysicsCharacterMovementModeInterface can be used with the CharacterPhysicsLiaison."),
 				*Input.SyncState.MovementMode.ToString());
 			return false;
diff --git a/Engine/Plugins/Experimental/Mover/Source/Mover/Private/DefaultMovementSet/LayeredMoves/AnimRootMotionLayeredMove.cpp b/Engine/Plugins/Experimental/Mover/Source/Mover/Private/DefaultMovementSet/LayeredMoves/AnimRootMotionLayeredMove.cpp
--- a/Engine/Plugins/Experimental/Mover/Source/Mover/Private/DefaultMovementSet/LayeredMoves/AnimRootMotionLayeredMove.cpp
+++ b/Engine/Plugins/Experimental/Mover/Source/Mover/Private/DefaultMovementSet/LayeredMoves/AnimRootMotionLayeredMove.cpp
@@ -52,8 +52,6 @@ bool FLayeredMove_AnimRootMotion::GenerateMove(const FMoverTickStartData& SimSta
 
 	const float DeltaSeconds = TimeStep.StepMs / 1000.f;
 
-	const AActor* MoverActor = MoverComp->GetOwner();
-
 	const FMoverDefaultSyncState* SyncState = SimState.SyncState.SyncStateCollection.FindDataByType<FMoverDefaultSyncState>();
 
 	// First pass simply samples based on the duration. For long animations, this has the potential to diverge.
diff --git a/Engine/Plugins/Experimental/Mover/Source/Mover/Private/LayeredMove.cpp b/Engine/Plugins/Experimental/Mover/Source/Mover/Private/LayeredMove.cpp
--- a/Engine/Plugins/Experimental/Mover/Source/Mover/Private/LayeredMove.cpp
+++ b/Engine/Plugins/Experimental/Mover/Source/Mover/Private/LayeredMove.cpp
@@ -44,6 +44,18 @@ void FLayeredMoveBase::StartMove(const UMoverComponent* MoverComp, UMoverBlackbo
 	OnStart(MoverComp, SimBlackboard);
 }
 
+void FLayeredMoveBase::StartMove_Async(UMoverBlackboard* SimBlackboard, float CurrentSimTimeMs)
+{
+	StartSimTimeMs = CurrentSimTimeMs;
+	OnStart_Async(SimBlackboard);
+}
+
+bool FLayeredMoveBase::GenerateMove_Async(const FMoverTickStartData& StartState, const FMoverTimeStep& TimeStep, UMoverBlackboard* SimBlackboard, FProposedMove& OutProposedMove)
+{
+	ensureMsgf(false, TEXT("GenerateMove_Async is not implemented"));
+	return false;
+}
+
 bool FLayeredMoveBase::IsFinished(float CurrentSimTimeMs) const
 {
 	const bool bHasStarted = (StartSimTimeMs >= 0.f);
@@ -58,6 +70,11 @@ void FLayeredMoveBase::EndMove(const UMoverComponent* MoverComp, UMoverBlackboar
 	OnEnd(MoverComp, SimBlackboard, CurrentSimTimeMs);
 }
 
+void FLayeredMoveBase::EndMove_Async(UMoverBlackboard* SimBlackboard, float CurrentSimTimeMs)
+{
+	OnEnd_Async(SimBlackboard, CurrentSimTimeMs);
+}
+
 FLayeredMoveBase* FLayeredMoveBase::Clone() const
 {
 	// If child classes don't override this, saved moves will not work
@@ -120,7 +137,17 @@ TArray<TSharedPtr<FLayeredMoveBase>> FLayeredMoveGroup::GenerateActiveMoves(cons
 	const float SimStartTimeMs		= TimeStep.BaseSimTimeMs;
 	const float SimTimeAfterTickMs	= SimStartTimeMs + TimeStep.StepMs;
 
-	FlushMoveArrays(MoverComp, SimBlackboard, SimStartTimeMs);
+	FlushMoveArrays(MoverComp, SimBlackboard, SimStartTimeMs, /*bIsAsync =*/ false);
+
+	return ActiveLayeredMoves;
+}
+
+TArray<TSharedPtr<FLayeredMoveBase>> FLayeredMoveGroup::GenerateActiveMoves_Async(const FMoverTimeStep& TimeStep, UMoverBlackboard* SimBlackboard)
+{
+	const float SimStartTimeMs		= TimeStep.BaseSimTimeMs;
+	const float SimTimeAfterTickMs	= SimStartTimeMs + TimeStep.StepMs;
+
+	FlushMoveArrays(nullptr, SimBlackboard, SimStartTimeMs, /*bIsAsync =*/ true);
 
 	return ActiveLayeredMoves;
 }
@@ -279,12 +306,18 @@ const FLayeredMoveBase* FLayeredMoveGroup::FindQueuedMove(const UScriptStruct* L
 	return nullptr;
 }
 
-void FLayeredMoveGroup::FlushMoveArrays(const UMoverComponent* MoverComp, UMoverBlackboard* SimBlackboard, float CurrentSimTimeMs)
+void FLayeredMoveGroup::FlushMoveArrays(const UMoverComponent* MoverComp, UMoverBlackboard* SimBlackboard, float CurrentSimTimeMs, bool bIsAsync)
 {
+	if (bIsAsync)
+	{
+		ensureMsgf(MoverComp == nullptr, TEXT("In async mode, no mover Component should be passed in as argument in FLayeredMoveGroup::FlushMoveArrays"));
+		MoverComp = nullptr;
+	}
+
 	bool bResidualVelocityOverridden = false;
 	bool bClampVelocityOverridden = false;
 	// Remove any finished moves
-	ActiveLayeredMoves.RemoveAll([MoverComp, SimBlackboard, CurrentSimTimeMs, &bResidualVelocityOverridden, &bClampVelocityOverridden, this]
+	ActiveLayeredMoves.RemoveAll([MoverComp, SimBlackboard, CurrentSimTimeMs, &bResidualVelocityOverridden, &bClampVelocityOverridden, bIsAsync, this]
 		(const TSharedPtr<FLayeredMoveBase>& Move)
 		{
 			if (Move.IsValid())
@@ -292,7 +325,14 @@ void FLayeredMoveGroup::FlushMoveArrays(const UMoverComponent* MoverComp, UMover
 				if (Move->IsFinished(CurrentSimTimeMs))
 				{
 					GatherResidualVelocitySettings(Move, bResidualVelocityOverridden, bClampVelocityOverridden);
-					Move->EndMove(MoverComp, SimBlackboard, CurrentSimTimeMs);
+					if (bIsAsync)
+					{
+						Move->EndMove_Async(SimBlackboard, CurrentSimTimeMs);
+					}
+					else
+					{
+						Move->EndMove(MoverComp, SimBlackboard, CurrentSimTimeMs);
+					}
 					return true;
 				}
 			}
@@ -308,7 +348,14 @@ void FLayeredMoveGroup::FlushMoveArrays(const UMoverComponent* MoverComp, UMover
 	for (TSharedPtr<FLayeredMoveBase>& QueuedMove : QueuedLayeredMoves)
 	{
 		ActiveLayeredMoves.Add(QueuedMove);
-		QueuedMove->StartMove(MoverComp, SimBlackboard, CurrentSimTimeMs);
+		if (bIsAsync)
+		{
+			QueuedMove->StartMove_Async(SimBlackboard, CurrentSimTimeMs);
+		}
+		else
+		{
+			QueuedMove->StartMove(MoverComp, SimBlackboard, CurrentSimTimeMs);
+		}
 	}
 
 	QueuedLayeredMoves.Empty();
diff --git a/Engine/Plugins/Experimental/Mover/Source/Mover/Private/MoveLibrary/MovementUtils.cpp b/Engine/Plugins/Experimental/Mover/Source/Mover/Private/MoveLibrary/MovementUtils.cpp
--- a/Engine/Plugins/Experimental/Mover/Source/Mover/Private/MoveLibrary/MovementUtils.cpp
+++ b/Engine/Plugins/Experimental/Mover/Source/Mover/Private/MoveLibrary/MovementUtils.cpp
@@ -160,6 +160,18 @@ FVector UMovementUtils::ComputeCombinedVelocity(const FComputeCombinedVelocityPa
 	return Velocity;
 }
 
+FVector UMovementUtils::DeduceUpDirectionFromGravity(const FVector& GravityAcceleration)
+{
+	const FVector DeducedUpDir = -GravityAcceleration.GetSafeNormal();
+
+	if (DeducedUpDir.IsZero())
+	{
+		return MoverComponentConstants::DefaultUpDir;
+	}
+
+	return DeducedUpDir;
+}
+
 bool UMovementUtils::CanEscapeGravity(const FVector& PriorVelocity, const FVector& NewVelocity, const FVector& GravityAccel, float DeltaSeconds)
 {
 	if (DeltaSeconds > UE_SMALL_NUMBER)
diff --git a/Engine/Plugins/Experimental/Mover/Source/Mover/Private/MovementMode.cpp b/Engine/Plugins/Experimental/Mover/Source/Mover/Private/MovementMode.cpp
--- a/Engine/Plugins/Experimental/Mover/Source/Mover/Private/MovementMode.cpp
+++ b/Engine/Plugins/Experimental/Mover/Source/Mover/Private/MovementMode.cpp
@@ -3,6 +3,7 @@
 #include "MovementMode.h"
 #include "MovementModeTransition.h"
 #include "MoverComponent.h"
+#include "MoverLog.h"
 #include "Engine/BlueprintGeneratedClass.h"
 
 #if WITH_EDITOR
@@ -103,4 +104,15 @@ bool UBaseMovementMode::HasGameplayTag(FGameplayTag TagToFind, bool bExactMatch)
 	return GameplayTags.HasTag(TagToFind);
 }
 
+const FName UNullMovementMode::NullModeName(TEXT("Null"));
+
+UNullMovementMode::UNullMovementMode(const FObjectInitializer& ObjectInitializer)
+	: Super(ObjectInitializer)
+{
+}
+
+void UNullMovementMode::SimulationTick_Implementation(const FSimulationTickParams& Params, FMoverTickEndData& OutputState)
+{
+}
+
 #undef LOCTEXT_NAMESPACE
diff --git a/Engine/Plugins/Experimental/Mover/Source/Mover/Private/MovementModeStateMachine.cpp b/Engine/Plugins/Experimental/Mover/Source/Mover/Private/MovementModeStateMachine.cpp
--- a/Engine/Plugins/Experimental/Mover/Source/Mover/Private/MovementModeStateMachine.cpp
+++ b/Engine/Plugins/Experimental/Mover/Source/Mover/Private/MovementModeStateMachine.cpp
@@ -24,20 +24,6 @@ namespace MoverComponentCVars
 
 } // end MoverComponentCVars
 
-
-const FName UNullMovementMode::NullModeName(TEXT("Null"));
-
-UNullMovementMode::UNullMovementMode(const FObjectInitializer& ObjectInitializer)
-	: Super(ObjectInitializer)
-{
-}
-
-void UNullMovementMode::SimulationTick_Implementation(const FSimulationTickParams& Params, FMoverTickEndData& OutputState)
-{
-}
-
-
-
 UMovementModeStateMachine::UMovementModeStateMachine(const FObjectInitializer& ObjectInitializer)
 	: Super(ObjectInitializer)
 {
@@ -316,6 +302,7 @@ void UMovementModeStateMachine::OnSimulationTick(USceneComponent* UpdatedCompone
 				FSimulationTickParams SimTickParams;
 				SimTickParams.StartState = SubstepStartData;
 				SimTickParams.MovingComps = FMovingComponentSet(MoverComp);
+				SimTickParams.SimBlackboard = SimBlackboard;
 				SimTickParams.TimeStep = SubTimeStep;
 				SimTickParams.ProposedMove = CombinedMove;
 
@@ -676,46 +663,3 @@ void UMovementModeStateMachine::PostInitProperties()
 		QueuedModeTransition = NewObject<UImmediateMovementModeTransition>(this, TEXT("QueuedModeTransition"), RF_Transient);
 	}
 }
-
-
-
-
-UImmediateMovementModeTransition::UImmediateMovementModeTransition(const FObjectInitializer& ObjectInitializer)
-	: Super(ObjectInitializer)
-{
-	Clear();
-}
-
-FTransitionEvalResult UImmediateMovementModeTransition::Evaluate_Implementation(const FSimulationTickParams& Params) const
-{
-	if (NextMode != NAME_None)
-	{
-		if (bAllowModeReentry)
-		{
-			return FTransitionEvalResult(NextMode);
-		}
-		else if (NextMode != Params.StartState.SyncState.MovementMode)
-		{
-			return FTransitionEvalResult(NextMode);
-		}
-	}
-
-	return FTransitionEvalResult::NoTransition;
-}
-
-void UImmediateMovementModeTransition::Trigger_Implementation(const FSimulationTickParams& Params)
-{
-	Clear();
-}
-
-void UImmediateMovementModeTransition::SetNextMode(FName DesiredModeName, bool bShouldReenter)
-{
-	NextMode = DesiredModeName;
-	bAllowModeReentry = bShouldReenter;
-}
-
-void UImmediateMovementModeTransition::Clear()
-{
-	NextMode = NAME_None;
-	bAllowModeReentry = false;
-}
diff --git a/Engine/Plugins/Experimental/Mover/Source/Mover/Private/MovementModeTransition.cpp b/Engine/Plugins/Experimental/Mover/Source/Mover/Private/MovementModeTransition.cpp
--- a/Engine/Plugins/Experimental/Mover/Source/Mover/Private/MovementModeTransition.cpp
+++ b/Engine/Plugins/Experimental/Mover/Source/Mover/Private/MovementModeTransition.cpp
@@ -51,3 +51,43 @@ EDataValidationResult UBaseMovementModeTransition::IsDataValid(FDataValidationCo
 	return EDataValidationResult::Valid;
 }
 #endif // WITH_EDITOR
+
+UImmediateMovementModeTransition::UImmediateMovementModeTransition(const FObjectInitializer& ObjectInitializer)
+	: Super(ObjectInitializer)
+{
+	Clear();
+}
+
+FTransitionEvalResult UImmediateMovementModeTransition::Evaluate_Implementation(const FSimulationTickParams& Params) const
+{
+	if (NextMode != NAME_None)
+	{
+		if (bAllowModeReentry)
+		{
+			return FTransitionEvalResult(NextMode);
+		}
+		else if (NextMode != Params.StartState.SyncState.MovementMode)
+		{
+			return FTransitionEvalResult(NextMode);
+		}
+	}
+
+	return FTransitionEvalResult::NoTransition;
+}
+
+void UImmediateMovementModeTransition::Trigger_Implementation(const FSimulationTickParams& Params)
+{
+	Clear();
+}
+
+void UImmediateMovementModeTransition::SetNextMode(FName DesiredModeName, bool bShouldReenter)
+{
+	NextMode = DesiredModeName;
+	bAllowModeReentry = bShouldReenter;
+}
+
+void UImmediateMovementModeTransition::Clear()
+{
+	NextMode = NAME_None;
+	bAllowModeReentry = false;
+}
diff --git a/Engine/Plugins/Experimental/Mover/Source/Mover/Private/MoverComponent.cpp b/Engine/Plugins/Experimental/Mover/Source/Mover/Private/MoverComponent.cpp
--- a/Engine/Plugins/Experimental/Mover/Source/Mover/Private/MoverComponent.cpp
+++ b/Engine/Plugins/Experimental/Mover/Source/Mover/Private/MoverComponent.cpp
@@ -1012,6 +1012,36 @@ bool UMoverComponent::ValidateSetup(FDataValidationContext& Context) const
 {
 	bool bHasMatchingStartingMode = false;
 	bool bDidFindAnyProblems = false;
+	bool bIsAsyncBackend = false;
+
+	// Verify backend liaison
+	if (!BackendClass)
+	{
+		Context.AddError(FText::Format(LOCTEXT("MissingBackendClassError", "No BackendClass property specified on {0}. Mover actor will not function."),
+			FText::FromString(GetNameSafe(GetOwner()))));
+
+		bDidFindAnyProblems = true;
+	}
+	else if (!BackendClass->ImplementsInterface(UMoverBackendLiaisonInterface::StaticClass()))
+	{
+		Context.AddError(FText::Format(LOCTEXT("InvalidBackendClassError", "BackendClass {0} on {1} does not implement IMoverBackendLiaisonInterface. Mover actor will not function."),
+			FText::FromString(BackendClass->GetName()),
+			FText::FromString(GetNameSafe(GetOwner()))));
+
+		bDidFindAnyProblems = true;
+	}
+	else
+	{
+		IMoverBackendLiaisonInterface* BackendCDOAsInterface = Cast<IMoverBackendLiaisonInterface>(BackendClass->GetDefaultObject());
+		if (BackendCDOAsInterface)
+		{
+			bIsAsyncBackend = BackendCDOAsInterface->IsAsync();
+			if (BackendCDOAsInterface->ValidateData(Context, *this) == EDataValidationResult::Invalid)
+			{
+				bDidFindAnyProblems = true;
+			}
+		}
+	}
 
 	// Verify all movement modes
 	for (const TPair<FName, TObjectPtr<UBaseMovementMode>>& Element : MovementModes)
@@ -1038,6 +1068,15 @@ bool UMoverComponent::ValidateSetup(FDataValidationContext& Context) const
 		// Verify that the movement mode's shared settings object exists (if any)
 		if (Element.Value)
 		{
+			if (bIsAsyncBackend && !Element.Value->bSupportsAsync)
+			{
+				Context.AddError(FText::Format(LOCTEXT("InvalidModeAsyncSupportsError", "Movement mode on {0}, mapped as {1} does not support asynchrony but its backend is asynchronous"),
+						FText::FromString(GetNameSafe(GetOwner())),
+						FText::FromName(Element.Key)));
+
+				bDidFindAnyProblems = true;
+			}
+
 			for (TSubclassOf<UObject>& Type : Element.Value->SharedSettingsClasses)
 			{
 				if (Type.Get() == nullptr)
@@ -1066,6 +1105,15 @@ bool UMoverComponent::ValidateSetup(FDataValidationContext& Context) const
 					continue;
 				}
 
+				if (bIsAsyncBackend && !Transition->bSupportsAsync)
+				{
+					Context.AddError(FText::Format(LOCTEXT("InvalidModeTransitionAsyncSupportError", "Transition on mode {0} on {1} does not support asynchrony but its backend is asynchronous"),
+						FText::FromName(Element.Key),
+						FText::FromString(GetNameSafe(GetOwner()))));
+
+					bDidFindAnyProblems = true;
+				}
+
 				for (const TSubclassOf<UObject>& Type : Transition->SharedSettingsClasses)
 				{
 					if (Type.Get() == nullptr)
@@ -1132,31 +1180,6 @@ bool UMoverComponent::ValidateSetup(FDataValidationContext& Context) const
 		}
 	}
 
-	// Verify backend liaison
-	if (!BackendClass)
-	{
-		Context.AddError(FText::Format(LOCTEXT("MissingBackendClassError", "No BackendClass property specified on {0}. Mover actor will not function."),
-			FText::FromString(GetNameSafe(GetOwner()))));
-
-		bDidFindAnyProblems = true;
-	}
-	else if (!BackendClass->ImplementsInterface(UMoverBackendLiaisonInterface::StaticClass()))
-	{
-		Context.AddError(FText::Format(LOCTEXT("InvalidBackendClassError", "BackendClass {0} on {1} does not implement IMoverBackendLiaisonInterface. Mover actor will not function."),
-			FText::FromString(BackendClass->GetName()),
-			FText::FromString(GetNameSafe(GetOwner()))));
-
-		bDidFindAnyProblems = true;
-	}
-	else
-	{
-		IMoverBackendLiaisonInterface* BackendCDOAsInterface = Cast<IMoverBackendLiaisonInterface>(BackendClass->GetDefaultObject());
-		if (BackendCDOAsInterface && (BackendCDOAsInterface->ValidateData(Context, *this) == EDataValidationResult::Invalid))
-		{
-			bDidFindAnyProblems = true;
-		}
-	}
-
 	// Verify persistent types
 	for (const FMoverDataPersistence& PersistentSyncEntry : PersistentSyncStateDataTypes)
 	{
@@ -1823,14 +1846,7 @@ FVector UMoverComponent::GetGravityAcceleration() const
 
 FVector UMoverComponent::GetUpDirection() const
 {
-	const FVector DeducedUpDir = -GetGravityAcceleration().GetSafeNormal();
-
-	if (DeducedUpDir.IsZero())
-	{
-		return MoverComponentConstants::DefaultUpDir;
-	}
-
-	return DeducedUpDir;
+	return UMovementUtils::DeduceUpDirectionFromGravity(GetGravityAcceleration());
 }
 
 const FPlanarConstraint& UMoverComponent::GetPlanarConstraint() const
@@ -2170,4 +2186,19 @@ UBaseMovementMode* UMoverComponent::GetActiveModeInternal(TSubclassOf<UBaseMovem
 	return nullptr;
 }
 
+void UMoverComponent::SetSimulationOutput(const FMoverTimeStep& TimeStep, const UE::Mover::FSimulationOutputData& OutputData)
+{
+	CachedLastSimTickTimeStep = TimeStep;
+
+	CachedLastUsedInputCmd = OutputData.LastUsedInputCmd;
+	bHasValidCachedUsedInput = true;
+
+	CachedLastSyncState = OutputData.SyncState;
+	bHasValidCachedState = true;
+
+	// This is for things like the ground info that we want to cache and interpolate but isn't part of the networked sync state
+	// AdditionalOutputData is generic because ground info might not be useful for platforms, say, but platforms may want to store something else.
+	// OutputData.AdditionalOutputData;
+}
+
 #undef LOCTEXT_NAMESPACE
diff --git a/Engine/Plugins/Experimental/Mover/Source/Mover/Public/Backends/MoverBackendLiaison.h b/Engine/Plugins/Experimental/Mover/Source/Mover/Public/Backends/MoverBackendLiaison.h
--- a/Engine/Plugins/Experimental/Mover/Source/Mover/Public/Backends/MoverBackendLiaison.h
+++ b/Engine/Plugins/Experimental/Mover/Source/Mover/Public/Backends/MoverBackendLiaison.h
@@ -29,6 +29,9 @@ public:
 	virtual float GetCurrentSimTimeMs() = 0;
 	virtual int32 GetCurrentSimFrame() = 0;
 
+	// Whether this backend will simulate movement asynchronously
+	virtual bool IsAsync() const { return false; }
+
 	// Pending State: the simulation state currently being authored
 	virtual bool ReadPendingSyncState(OUT FMoverSyncState& OutSyncState) { return false; }
 	virtual bool WritePendingSyncState(const FMoverSyncState& SyncStateToWrite) { return false; }
diff --git a/Engine/Plugins/Experimental/Mover/Source/Mover/Public/InstantMovementEffect.h b/Engine/Plugins/Experimental/Mover/Source/Mover/Public/InstantMovementEffect.h
--- a/Engine/Plugins/Experimental/Mover/Source/Mover/Public/InstantMovementEffect.h
+++ b/Engine/Plugins/Experimental/Mover/Source/Mover/Public/InstantMovementEffect.h
@@ -23,6 +23,16 @@ struct FApplyMovementEffectParams
 	const FMoverTimeStep* TimeStep;
 };
 
+/** 
+ * Async safe parameters passed to ApplyMovementEffect_Async. 
+ * It is almost certainly missing the Physics Object handle and other things, this is just a first pass
+ */
+struct FApplyMovementEffectParams_Async
+{
+	const FMoverTickStartData* StartState;
+	const FMoverTimeStep* TimeStep;
+};
+
 /**
  * Instant Movement Effects are methods of affecting movement state directly on a Mover-based actor for one tick.
  * Note: This is only applied one tick and then removed
@@ -50,6 +60,7 @@ struct MOVER_API FInstantMovementEffect
 	virtual void AddReferencedObjects(class FReferenceCollector& Collector) {}
 
 	virtual bool ApplyMovementEffect(FApplyMovementEffectParams& ApplyEffectParams, FMoverSyncState& OutputState) { return false; }
+	virtual bool ApplyMovementEffect_Async(FApplyMovementEffectParams_Async& ApplyEffectParams, FMoverSyncState& OutputState) { return false; }
 };
 
 template<>
diff --git a/Engine/Plugins/Experimental/Mover/Source/Mover/Public/LayeredMove.h b/Engine/Plugins/Experimental/Mover/Source/Mover/Public/LayeredMove.h
--- a/Engine/Plugins/Experimental/Mover/Source/Mover/Public/LayeredMove.h
+++ b/Engine/Plugins/Experimental/Mover/Source/Mover/Public/LayeredMove.h
@@ -112,22 +112,32 @@ struct MOVER_API FLayeredMoveBase
 	
 	// Kicks off this move, allowing any initialization to occur.
 	void StartMove(const UMoverComponent* MoverComp, UMoverBlackboard* SimBlackboard, float CurrentSimTimeMs);
+	// Async version of the above
+	void StartMove_Async(UMoverBlackboard* SimBlackboard, float CurrentSimTimeMs);
 
 	// Called when this layered move starts.
 	virtual void OnStart(const UMoverComponent* MoverComp, UMoverBlackboard* SimBlackboard) {}
+	// Async version of the above
+	virtual void OnStart_Async(UMoverBlackboard* SimBlackboard) {}
 
 	// TODO: consider whether MoverComp should just be part of FMoverTickStartData
 	// Generate a movement that will be combined with other sources
 	virtual bool GenerateMove(const FMoverTickStartData& StartState, const FMoverTimeStep& TimeStep, const UMoverComponent* MoverComp, UMoverBlackboard* SimBlackboard, FProposedMove& OutProposedMove) { return false; }
+	// Async version of the above
+	virtual bool GenerateMove_Async(const FMoverTickStartData& StartState, const FMoverTimeStep& TimeStep, UMoverBlackboard* SimBlackboard, FProposedMove& OutProposedMove);
 
 	// Runtime query whether this move is finished and can be destroyed. The default implementation is based on DurationMs.
 	virtual bool IsFinished(float CurrentSimTimeMs) const;
 
 	// Ends this move, allowing any cleanup to occur.
 	void EndMove(const UMoverComponent* MoverComp, UMoverBlackboard* SimBlackboard, float CurrentSimTimeMs);
+	// Async version of the above
+	void EndMove_Async(UMoverBlackboard* SimBlackboard, float CurrentSimTimeMs);
 	
 	// Called when this layered move ends.
 	virtual void OnEnd(const UMoverComponent* MoverComp, UMoverBlackboard* SimBlackboard, float CurrentSimTimeMs) {}
+	// Async version of the above
+	virtual void OnEnd_Async(UMoverBlackboard* SimBlackboard, float CurrentSimTimeMs) {}
 
 	// @return newly allocated copy of this FLayeredMoveBase. Must be overridden by child classes
 	virtual FLayeredMoveBase* Clone() const;
@@ -180,6 +190,7 @@ struct MOVER_API FLayeredMoveGroup
 
 	// Generates active layered move list (by calling FlushMoveArrays) and returns the an array of all currently active layered moves
 	TArray<TSharedPtr<FLayeredMoveBase>> GenerateActiveMoves(const FMoverTimeStep& TimeStep, const UMoverComponent* MoverComp, UMoverBlackboard* SimBlackboard);
+	TArray<TSharedPtr<FLayeredMoveBase>> GenerateActiveMoves_Async(const FMoverTimeStep& TimeStep, UMoverBlackboard* SimBlackboard);
 
 	/** Serialize all moves and their states for this group */
 	void NetSerialize(FArchive& Ar, uint8 MaxNumMovesToSerialize = MAX_uint8);
@@ -226,7 +237,7 @@ struct MOVER_API FLayeredMoveGroup
 
 protected:
 	// Clears out any finished or invalid active moves and adds any queued moves to the active moves
-	void FlushMoveArrays(const UMoverComponent* MoverComp, UMoverBlackboard* SimBlackboard, float CurrentSimTimeMs);
+	void FlushMoveArrays(const UMoverComponent* MoverComp, UMoverBlackboard* SimBlackboard, float CurrentSimTimeMs, bool bIsAsync);
 
 	// Helper function for gathering any residual velocity settings from layered moves that just ended
 	void GatherResidualVelocitySettings(const TSharedPtr<FLayeredMoveBase>& Move, bool& bResidualVelocityOverriden, bool& bClampVelocityOverriden);
diff --git a/Engine/Plugins/Experimental/Mover/Source/Mover/Public/MoveLibrary/MovementUtils.h b/Engine/Plugins/Experimental/Mover/Source/Mover/Public/MoveLibrary/MovementUtils.h
--- a/Engine/Plugins/Experimental/Mover/Source/Mover/Public/MoveLibrary/MovementUtils.h
+++ b/Engine/Plugins/Experimental/Mover/Source/Mover/Public/MoveLibrary/MovementUtils.h
@@ -149,6 +149,10 @@ public:
 	UFUNCTION(BlueprintCallable, Category=Mover)
 	static FVector ComputeVelocityFromGravity(const FVector& GravityAccel, float DeltaSeconds) { return GravityAccel * DeltaSeconds; }
 
+	/** Returns the up direction deduced from gravity acceleration, but defaults to mover constants up direction if zero */
+	UFUNCTION(BlueprintCallable, Category=Mover)
+	static FVector DeduceUpDirectionFromGravity(const FVector& GravityAcceleration);
+
 	/** Checks whether a given velocity is strong enough to lift off against gravity */
 	UFUNCTION(BlueprintCallable, Category=Mover)
 	static bool CanEscapeGravity(const FVector& PriorVelocity, const FVector& NewVelocity, const FVector& GravityAccel, float DeltaSeconds);
diff --git a/Engine/Plugins/Experimental/Mover/Source/Mover/Public/MovementMode.h b/Engine/Plugins/Experimental/Mover/Source/Mover/Public/MovementMode.h
--- a/Engine/Plugins/Experimental/Mover/Source/Mover/Public/MovementMode.h
+++ b/Engine/Plugins/Experimental/Mover/Source/Mover/Public/MovementMode.h
@@ -102,7 +102,14 @@ public:
 	/** A list of gameplay tags associated with this movement mode */
 	UPROPERTY(EditDefaultsOnly, Category = Mover)
 	FGameplayTagContainer GameplayTags;
-	
+
+	/** 
+	 * Whether this movement mode supports being part of an asynchronous movement simulation (running concurrently with the gameplay thread) 
+	 * Specifically for the GenerateMove and SimulationTick functions
+	 */
+	UPROPERTY(EditDefaultsOnly, BlueprintReadOnly, Category = Mover)
+	bool bSupportsAsync = false;
+
 protected:
 	UFUNCTION(BlueprintImplementableEvent, meta = (DisplayName = "On Activated", ScriptName = "OnActivated"))
 	void K2_OnActivated();
@@ -124,4 +131,18 @@ protected:
 	virtual void OnGenerateMove(const FMoverTickStartData& StartState, const FMoverTimeStep& TimeStep, FProposedMove& OutProposedMove) const final {}
 	UE_DEPRECATED(5.6, "OnSimulationTick() has been replaced with a SimulationTick() BlueprintNativeEvent. Rename your override to SimulationTick_Implementation().")
 	virtual void OnSimulationTick(const FSimulationTickParams& Params, FMoverTickEndData& OutputState) final {}
-};
\ No newline at end of file
+};
+
+/**
+ * NullMovementMode: a default do-nothing mode used as a placeholder when no other mode is active
+ */
+ UCLASS(NotBlueprintable)
+class MOVER_API UNullMovementMode : public UBaseMovementMode
+{
+	GENERATED_UCLASS_BODY()
+
+public:
+	virtual void SimulationTick_Implementation(const FSimulationTickParams& Params, FMoverTickEndData& OutputState) override;
+
+	const static FName NullModeName;
+};
diff --git a/Engine/Plugins/Experimental/Mover/Source/Mover/Public/MovementModeStateMachine.h b/Engine/Plugins/Experimental/Mover/Source/Mover/Public/MovementModeStateMachine.h
--- a/Engine/Plugins/Experimental/Mover/Source/Mover/Public/MovementModeStateMachine.h
+++ b/Engine/Plugins/Experimental/Mover/Source/Mover/Public/MovementModeStateMachine.h
@@ -12,20 +12,6 @@ struct FProposedMove;
 class UImmediateMovementModeTransition;
 class UMovementModeTransition;
 
-/**
- * NullMovementMode: a default do-nothing mode used as a placeholder when no other mode is active
- */
- UCLASS(NotBlueprintable)
-class UNullMovementMode : public UBaseMovementMode
-{
-	GENERATED_UCLASS_BODY()
-
-public:
-	virtual void SimulationTick_Implementation(const FSimulationTickParams& Params, FMoverTickEndData& OutputState) override;
-
-	const static FName NullModeName;
-};
-
 /**
  * - Any movement modes registered are co-owned by the state machine
  * - There is always an active mode, falling back to a do-nothing 'null' mode
@@ -42,7 +28,6 @@ public:
 	void RegisterMovementMode(FName ModeName, TSubclassOf<UBaseMovementMode> ModeType, bool bIsDefaultMode=false);
 
 	void UnregisterMovementMode(FName ModeName);
-	void UnregisterMovementMode(TObjectPtr<UBaseMovementMode> ModeToUnregister);
 	void ClearAllMovementModes();
 
 	void RegisterGlobalTransition(TObjectPtr<UBaseMovementModeTransition> Transition);
@@ -109,31 +94,3 @@ private:
  	bool ApplyInstantEffects(FApplyMovementEffectParams& ApplyEffectParams, FMoverSyncState& OutputState);
 	AActor* GetOwnerActor() const;
 };
-
-
-
- /**
-  * Simple transition that evaluates true if a "next mode" is set. Used internally only by the Mover plugin. 
-  */
- UCLASS()
- class UImmediateMovementModeTransition : public UBaseMovementModeTransition
- {
-	 GENERATED_UCLASS_BODY()
-
- public:
-	 virtual FTransitionEvalResult Evaluate_Implementation(const FSimulationTickParams& Params) const override;
-	 virtual void Trigger_Implementation(const FSimulationTickParams& Params) override;
-
-
-	 bool IsSet() const { return !NextMode.IsNone(); }
-	 void SetNextMode(FName DesiredModeName, bool bShouldReenter = false);
-	 void Clear();
-
-
-
-	 FName GetNextModeName() const { return NextMode; }
-	 bool ShouldReenter() const { return bAllowModeReentry; }
-
- private:
-	 FName NextMode = NAME_None;
- };
diff --git a/Engine/Plugins/Experimental/Mover/Source/Mover/Public/MovementModeTransition.h b/Engine/Plugins/Experimental/Mover/Source/Mover/Public/MovementModeTransition.h
--- a/Engine/Plugins/Experimental/Mover/Source/Mover/Public/MovementModeTransition.h
+++ b/Engine/Plugins/Experimental/Mover/Source/Mover/Public/MovementModeTransition.h
@@ -82,7 +82,14 @@ public:
 	*/
 	UPROPERTY(EditDefaultsOnly, BlueprintReadWrite, Category = Mover)
 	bool bFirstSubStepOnly = false;
-	
+
+	/** 
+	 * Whether this movement mode transition supports being part of an asynchronous movement simulation (running concurrently with the gameplay thread) 
+	 * Specifically for the Evaluate and Trigger functions
+	 */
+	UPROPERTY(EditDefaultsOnly, BlueprintReadOnly, Category = Mover)
+	bool bSupportsAsync = false;
+
 protected:
 	UFUNCTION(BlueprintImplementableEvent, Category = MovementModeTransition, meta = (DisplayName = "On Registered", ScriptName = "OnRegistered"))
 	void K2_OnRegistered();
@@ -96,6 +103,28 @@ protected:
 	virtual void OnTrigger(const FSimulationTickParams& Params) final {}
 };
 
+/**
+ * Simple transition that evaluates true if a "next mode" is set. Used internally only by the Mover plugin. 
+ */
+UCLASS()
+class MOVER_API UImmediateMovementModeTransition : public UBaseMovementModeTransition
+{
+	GENERATED_UCLASS_BODY()
+
+public:
+	virtual FTransitionEvalResult Evaluate_Implementation(const FSimulationTickParams& Params) const override;
+	virtual void Trigger_Implementation(const FSimulationTickParams& Params) override;
+
+	bool IsSet() const { return !NextMode.IsNone(); }
+	void SetNextMode(FName DesiredModeName, bool bShouldReenter = false);
+	void Clear();
+
+	FName GetNextModeName() const { return NextMode; }
+	bool ShouldReenter() const { return bAllowModeReentry; }
+
+private:
+	FName NextMode = NAME_None;
+};
 
 
 
diff --git a/Engine/Plugins/Experimental/Mover/Source/Mover/Public/MoverComponent.h b/Engine/Plugins/Experimental/Mover/Source/Mover/Public/MoverComponent.h
--- a/Engine/Plugins/Experimental/Mover/Source/Mover/Public/MoverComponent.h
+++ b/Engine/Plugins/Experimental/Mover/Source/Mover/Public/MoverComponent.h
@@ -167,7 +167,7 @@ public:
 	/** Optional object for mixing proposed moves.Typically set at BeginPlay time. If not specified, UDefaultMovementMixer will be used. */
 	UPROPERTY(BlueprintReadWrite, EditAnywhere, Category = Mover)
 	TObjectPtr<UMovementMixer> MovementMixer;
-	
+
 	/**
 	 * Queue a layered move to start during the next simulation frame. This will clone whatever move you pass in, so you'll need to fully set it up before queuing.
 	 * @param LayeredMove			The move to queue, which must be a LayeredMoveBase sub-type. 
@@ -619,6 +619,8 @@ public:
 	UPROPERTY(EditAnywhere, BlueprintReadWrite, Category = "Mover")
 	EMoverSmoothingMode SmoothingMode = EMoverSmoothingMode::VisualComponentOffset;
 
+	void SetSimulationOutput(const FMoverTimeStep& TimeStep, const UE::Mover::FSimulationOutputData& OutputData);
+
 protected:
 	bool bHasValidLastProducedInput = false;
 	bool bHasValidCachedUsedInput = false;
diff --git a/Engine/Plugins/Experimental/Mover/Source/Mover/Public/MoverSimulationTypes.h b/Engine/Plugins/Experimental/Mover/Source/Mover/Public/MoverSimulationTypes.h
--- a/Engine/Plugins/Experimental/Mover/Source/Mover/Public/MoverSimulationTypes.h
+++ b/Engine/Plugins/Experimental/Mover/Source/Mover/Public/MoverSimulationTypes.h
@@ -188,20 +188,27 @@ public:
 
 	bool ShouldReconcile(const FMoverAuxStateContext& AuthorityState) const
 	{ 
-		return false; 
+		return AuxStateCollection.ShouldReconcile(AuthorityState.AuxStateCollection); 
 	}
 
 	void NetSerialize(const FNetSerializeParams& P)
 	{
+		bool bIgnoredResult(false);
+		AuxStateCollection.NetSerialize(P.Ar, nullptr, bIgnoredResult);
 	}
 
 	void ToString(FAnsiStringBuilderBase& Out) const
 	{
+		AuxStateCollection.ToString(Out);
 	}
 
-	void Interpolate(const FMoverAuxStateContext* From, const FMoverAuxStateContext* To, float PCT)
+	void Interpolate(const FMoverAuxStateContext* From, const FMoverAuxStateContext* To, float Pct)
 	{
+		AuxStateCollection.Interpolate(From->AuxStateCollection, To->AuxStateCollection, Pct);
 	}
+
+	UPROPERTY(BlueprintReadWrite, Category = Mover)
+	FMoverDataCollection AuxStateCollection;
 };
 
 
@@ -264,9 +271,14 @@ struct MOVER_API FSimulationTickParams
 	GENERATED_BODY()
 
 	// Components involved in movement by the simulation
+	// This will be empty when the simulation is ticked asynchronously
 	UPROPERTY(BlueprintReadOnly, EditAnywhere, Category = Mover)
 	FMovingComponentSet MovingComps;
 
+	// Blackboard
+	UPROPERTY(BlueprintReadOnly, EditAnywhere, Category = Mover)
+	TObjectPtr<UMoverBlackboard> SimBlackboard;
+
 	// Simulation state data at the start of the tick, including Input Cmd
 	UPROPERTY(BlueprintReadOnly, EditAnywhere, Category = Mover)
 	FMoverTickStartData StartState;
@@ -280,7 +292,6 @@ struct MOVER_API FSimulationTickParams
 	FProposedMove ProposedMove;
 };
 
-
 UINTERFACE(BlueprintType)
 class MOVER_API UMoverInputProducerInterface : public UInterface
 {
@@ -342,4 +353,97 @@ struct MOVER_API FMoverPredictTrajectoryParams
  	UPROPERTY(EditAnywhere, BlueprintReadWrite, Category = Mover)
  	TArray<FMoverInputCmdContext> OptionalInputCmds;
 
-};
\ No newline at end of file
+};
+
+namespace UE::Mover
+{
+	struct FSimulationOutputData
+	{
+		void Reset()
+		{
+			SyncState.Reset();
+			LastUsedInputCmd.Reset();
+			AdditionalOutputData.Empty();
+		}
+
+		void Interpolate(const FSimulationOutputData& From, const FSimulationOutputData& To, float Alpha)
+		{
+			SyncState.Interpolate(&From.SyncState, &To.SyncState, Alpha);
+			LastUsedInputCmd.Interpolate(&From.LastUsedInputCmd, &To.LastUsedInputCmd, Alpha);
+			AdditionalOutputData.Interpolate(From.AdditionalOutputData, To.AdditionalOutputData, Alpha);
+		}
+
+		FMoverSyncState SyncState;
+		FMoverInputCmdContext LastUsedInputCmd;
+		FMoverDataCollection AdditionalOutputData;
+	};
+
+	class FSimulationOutputRecord
+	{
+	public:
+		struct FData
+		{
+			void Reset()
+			{
+				TimeStep = FMoverTimeStep();
+				SimOutputData.Reset();
+			}
+
+			FMoverTimeStep TimeStep;
+			FSimulationOutputData SimOutputData;
+		};
+		void Add(const FMoverTimeStep& InTimeStep , const FSimulationOutputData& InData)
+		{
+			CurrentIndex = (CurrentIndex + 1) % 2;
+			Data[CurrentIndex] = { InTimeStep, InData };
+		}
+
+		const FSimulationOutputData& GetLatest() const
+		{
+			return Data[CurrentIndex].SimOutputData;
+		}
+
+		void GetInterpolated(float SimTime, FMoverTimeStep& OutTimeStep, FSimulationOutputData& OutData) const
+		{
+			const uint8 PrevIndex = (CurrentIndex + 1) % 2;
+			const float PrevTime = Data[PrevIndex].TimeStep.BaseSimTimeMs * 0.001f;
+			const float CurrTime = Data[CurrentIndex].TimeStep.BaseSimTimeMs * 0.001f;
+
+			if (FMath::IsNearlyEqual(PrevTime, CurrTime) || (SimTime >= CurrTime))
+			{
+				OutData = Data[CurrentIndex].SimOutputData;
+				OutTimeStep = Data[CurrentIndex].TimeStep;
+			}
+			else if (SimTime <= PrevTime)
+			{
+				OutData = Data[PrevIndex].SimOutputData;
+				OutTimeStep = Data[PrevIndex].TimeStep;
+			}
+			else
+			{
+				const float Alpha = FMath::Clamp((SimTime - PrevTime) / (CurrTime - PrevTime), 0.0f, 1.0f);
+				OutData.Interpolate(Data[PrevIndex].SimOutputData, Data[CurrentIndex].SimOutputData, Alpha);
+				if (Alpha < 0.5f)
+				{
+					OutTimeStep = Data[PrevIndex].TimeStep;
+				}
+				else
+				{
+					OutTimeStep = Data[CurrentIndex].TimeStep;
+				}
+				OutTimeStep.BaseSimTimeMs = SimTime;
+			}
+		}
+
+		void Clear()
+		{
+			CurrentIndex = 1;
+			Data[0].Reset();
+			Data[1].Reset();
+		}
+
+		FData Data[2];
+		uint8 CurrentIndex = 1;
+	};
+
+} // namespace UE::Mover
\ No newline at end of file
