Chaos Mover
- Add character ground constraint to simulation
- Add falling movement mode
- Make simulation a UObject
- Add landing transition

diff --git a/Engine/Plugins/Experimental/ChaosMover/Source/ChaosMover/Private/Backends/ChaosMoverAsyncCallback.cpp b/Engine/Plugins/Experimental/ChaosMover/Source/ChaosMover/Private/Backends/ChaosMoverAsyncCallback.cpp
--- a/Engine/Plugins/Experimental/ChaosMover/Source/ChaosMover/Private/Backends/ChaosMoverAsyncCallback.cpp
+++ b/Engine/Plugins/Experimental/ChaosMover/Source/ChaosMover/Private/Backends/ChaosMoverAsyncCallback.cpp
@@ -11,11 +11,6 @@
 
 namespace UE::ChaosMover
 {
-	FAsyncCallback::~FAsyncCallback()
-	{
-		Simulations.Empty();
-	}
-
 	void FAsyncCallback::OnPreSimulate_Internal()
 	{
 		const FAsyncCallbackInput* AsyncInput = GetConsumerInput_Internal();
@@ -28,11 +23,11 @@ namespace UE::ChaosMover
 		FMoverTimeStep TimeStep;
 		TimeStep.BaseSimTimeMs = GetSimTime_Internal() * 1000.0f;
 		TimeStep.StepMs = GetDeltaTime_Internal() * 1000.0f;
-		if (PhysicsSolver)
+		if (AsyncInput->PhysicsSolver)
 		{
-			TimeStep.ServerFrame = PhysicsSolver->GetCurrentFrame() + NetworkPhysicsTickOffset;
-			TimeStep.BaseSimTimeMs += NetworkPhysicsTickOffset * TimeStep.StepMs;
-			TimeStep.bIsResimulating = PhysicsSolver->GetEvolution()->IsResimming();
+			TimeStep.ServerFrame = AsyncInput->PhysicsSolver->GetCurrentFrame() + AsyncInput->NetworkPhysicsTickOffset;
+			TimeStep.BaseSimTimeMs += AsyncInput->NetworkPhysicsTickOffset * TimeStep.StepMs;
+			TimeStep.bIsResimulating = AsyncInput->PhysicsSolver->GetEvolution()->IsResimming();
 		}
 
 		FAsyncCallbackOutput& AsyncOutput = GetProducerOutputData_Internal();
@@ -41,26 +36,29 @@ namespace UE::ChaosMover
 
 		auto LambdaParallelUpdate = [&TimeStep, &AsyncInput, &AsyncOutput, this](int32 Idx) {
 			AsyncOutput.TimeStep[Idx] = TimeStep;
-			if (UE::ChaosMover::FSimulation* Simulation = Simulations[Idx].Get())
+			if (TStrongObjectPtr<UChaosMoverBackendComponent> Backend = AsyncInput->Backends[Idx].Pin())
 			{
-				Simulation->SimulationTick(TimeStep, AsyncInput->InputData[Idx], AsyncOutput.OutputData[Idx]);
+				Backend->GetSimulation()->SimulationTick(TimeStep, AsyncInput->InputData[Idx], AsyncOutput.OutputData[Idx]);
+				AsyncOutput.Backends.Add(Backend.Get());
 			}
 		};
 
 		const bool ForceSingleThread = UE::ChaosMover::CVars::bForceSingleThreadedPT;
-		Chaos::PhysicsParallelFor(Simulations.Num(), LambdaParallelUpdate, ForceSingleThread);
+		Chaos::PhysicsParallelFor(AsyncInput->Backends.Num(), LambdaParallelUpdate, ForceSingleThread);
 	}
 
 	void FAsyncCallback::OnContactModification_Internal(Chaos::FCollisionContactModifier& Modifier)
 	{
-		auto LambdaParallelUpdate = [&Modifier, this](int32 Idx) {
-			if (UE::ChaosMover::FSimulation* Simulation = Simulations[Idx].Get())
+		const FAsyncCallbackInput* AsyncInput = GetConsumerInput_Internal();
+
+		auto LambdaParallelUpdate = [&Modifier, &AsyncInput](int32 Idx) {
+			if (TStrongObjectPtr<const UChaosMoverBackendComponent> Backend = AsyncInput->Backends[Idx].Pin())
 			{
-				Simulation->ModifyContacts(Modifier);
+				Backend->GetSimulation()->ModifyContacts(Modifier);
 			}
 		};
 
 		const bool ForceSingleThread = UE::ChaosMover::CVars::bForceSingleThreadedPT;
-		Chaos::PhysicsParallelFor(Simulations.Num(), LambdaParallelUpdate, ForceSingleThread);
+		Chaos::PhysicsParallelFor(AsyncInput->Backends.Num(), LambdaParallelUpdate, ForceSingleThread);
 	}
 }
\ No newline at end of file
diff --git a/Engine/Plugins/Experimental/ChaosMover/Source/ChaosMover/Private/Backends/ChaosMoverAsyncCallback.h b/Engine/Plugins/Experimental/ChaosMover/Source/ChaosMover/Private/Backends/ChaosMoverAsyncCallback.h
--- a/Engine/Plugins/Experimental/ChaosMover/Source/ChaosMover/Private/Backends/ChaosMoverAsyncCallback.h
+++ b/Engine/Plugins/Experimental/ChaosMover/Source/ChaosMover/Private/Backends/ChaosMoverAsyncCallback.h
@@ -9,24 +9,31 @@ class UChaosMoverBackendComponent;
 
 namespace UE::ChaosMover
 {
-	class FSimulation;
-
 	struct FAsyncCallbackInput : public Chaos::FSimCallbackInput
 	{
 		TArray<FSimulationInputData> InputData;
+		TArray<TWeakObjectPtr<UChaosMoverBackendComponent>> Backends;
+		Chaos::FPhysicsSolver* PhysicsSolver;
+		int32 NetworkPhysicsTickOffset;
+
 		void Reset()
 		{
 			InputData.Empty();
+			Backends.Empty();
+			PhysicsSolver = nullptr;
+			NetworkPhysicsTickOffset = 0;
 		}
 	};
 
 	struct FAsyncCallbackOutput : public Chaos::FSimCallbackOutput
 	{
 		TArray<FSimulationOutputData> OutputData;
+		TArray<TWeakObjectPtr<UChaosMoverBackendComponent>> Backends;
 		TArray<FMoverTimeStep> TimeStep;
 		void Reset()
 		{
 			OutputData.Empty();
+			Backends.Empty();
 			TimeStep.Empty();
 		}
 	};
@@ -38,14 +45,6 @@ namespace UE::ChaosMover
 		Chaos::ESimCallbackOptions::ContactModification |
 		Chaos::ESimCallbackOptions::Rewind>
 	{
-	public:
-		virtual ~FAsyncCallback();
-
-		TArray<TSharedPtr<UE::ChaosMover::FSimulation>> Simulations;
-		TArray<TWeakObjectPtr<UChaosMoverBackendComponent>> Backends;
-		Chaos::FPhysicsSolver* PhysicsSolver;
-		int32 NetworkPhysicsTickOffset = 0;
-
 	protected:
 		virtual void OnPreSimulate_Internal() override;
 		virtual void OnContactModification_Internal(Chaos::FCollisionContactModifier& Modifier) override;
diff --git a/Engine/Plugins/Experimental/ChaosMover/Source/ChaosMover/Private/Backends/ChaosMoverBackend.cpp b/Engine/Plugins/Experimental/ChaosMover/Source/ChaosMover/Private/Backends/ChaosMoverBackend.cpp
--- a/Engine/Plugins/Experimental/ChaosMover/Source/ChaosMover/Private/Backends/ChaosMoverBackend.cpp
+++ b/Engine/Plugins/Experimental/ChaosMover/Source/ChaosMover/Private/Backends/ChaosMoverBackend.cpp
@@ -3,19 +3,26 @@
 #include "ChaosMover/Backends/ChaosMoverBackend.h"
 
 #include "Backends/ChaosMoverSubsystem.h"
-#include "NetworkChaosMoverData.h"
-
+#include "ChaosMover/ChaosMoverLog.h"
 #include "ChaosMover/ChaosMoverSimulation.h"
 #include "Components/PrimitiveComponent.h"
 #include "Framework/Threading.h"
 #include "GameFramework/Pawn.h"
+#include "GameFramework/PhysicsVolume.h"
+#include "MoveLibrary/MovementMixer.h"
+#include "MoveLibrary/MoverBlackboard.h"
+#include "MovementModeStateMachine.h"
+#include "NetworkChaosMoverData.h"
 #include "PBDRigidsSolver.h"
 #include "Physics/Experimental/PhysScene_Chaos.h"
 #include "Physics/NetworkPhysicsComponent.h"
 #include "PhysicsEngine/PhysicsObjectExternalInterface.h"
-#include "MovementModeStateMachine.h"
-#include "MoveLibrary/MovementMixer.h"
-#include "MoveLibrary/MoverBlackboard.h"
+#include "PhysicsProxy/CharacterGroundConstraintProxy.h"
+
+#if WITH_CHAOS_VISUAL_DEBUGGER
+#include "ChaosVisualDebugger/ChaosVisualDebuggerTrace.h"
+#include "ChaosVisualDebugger/MoverCVDRuntimeTrace.h"
+#endif
 
 #include UE_INLINE_GENERATED_CPP_BY_NAME(ChaosMoverBackend)
 
@@ -35,50 +42,53 @@ UChaosMoverBackendComponent::UChaosMoverBackendComponent()
 	{
 		SetIsReplicatedByDefault(true);
 	}
+
+	Simulation = CreateDefaultSubobject<UChaosMoverSimulation>(TEXT("ChaosMoverSimulation"));
 }
 
 void UChaosMoverBackendComponent::InitializeComponent()
 {
 	Super::InitializeComponent();
 
-	// Create NetworkPhysicsComponent
-	if (Chaos::FPhysicsSolverBase::IsNetworkPhysicsPredictionEnabled() && bIsUsingAsyncPhysics)
+	if (UWorld* World = GetWorld(); World && World->IsGameWorld())
 	{
-		NetworkPhysicsComponent = NewObject<UNetworkPhysicsComponent>(GetOwner(), TEXT("PhysMover_NetworkPhysicsComponent"));
+		NullMovementMode = NewObject<UNullMovementMode>(&GetMoverComponent(), TEXT("NullMovementMode"));
+		ImmediateModeTransition = NewObject<UImmediateMovementModeTransition>(&GetMoverComponent(), TEXT("ImmediateModeTransition"));
 
-		// This isn't technically a DSO component, but set it net addressable as though it is
-		NetworkPhysicsComponent->SetNetAddressable();
-		NetworkPhysicsComponent->SetIsReplicated(true);
-		NetworkPhysicsComponent->RegisterComponent();
-		if (!NetworkPhysicsComponent->HasBeenInitialized())
+		// Create NetworkPhysicsComponent
+		if (Chaos::FPhysicsSolverBase::IsNetworkPhysicsPredictionEnabled() && bIsUsingAsyncPhysics)
 		{
-			NetworkPhysicsComponent->InitializeComponent();
-		}
-		NetworkPhysicsComponent->Activate(true);
+			NetworkPhysicsComponent = NewObject<UNetworkPhysicsComponent>(GetOwner(), TEXT("PhysMover_NetworkPhysicsComponent"));
 
-		// Register network data for recording and rewind/resim
-		NetworkPhysicsComponent->CreateDataHistory<UE::ChaosMover::FNetworkDataTraits>(this);
-
-		if (NetworkPhysicsComponent->HasServerWorld())
-		{
-			if (APawn* PawnOwner = Cast<APawn>(GetOwner()))
+			// This isn't technically a DSO component, but set it net addressable as though it is
+			NetworkPhysicsComponent->SetNetAddressable();
+			NetworkPhysicsComponent->SetIsReplicated(true);
+			NetworkPhysicsComponent->RegisterComponent();
+			if (!NetworkPhysicsComponent->HasBeenInitialized())
 			{
-				// When we're owned by a pawn, keep an eye on whether it's currently player-controlled or not
-				PawnOwner->ReceiveControllerChangedDelegate.AddUniqueDynamic(this, &ThisClass::HandleOwningPawnControllerChanged_Server);
-				HandleOwningPawnControllerChanged_Server(PawnOwner, nullptr, PawnOwner->Controller);
+				NetworkPhysicsComponent->InitializeComponent();
 			}
-			else
+			NetworkPhysicsComponent->Activate(true);
+
+			// Register network data for recording and rewind/resim
+			NetworkPhysicsComponent->CreateDataHistory<UE::ChaosMover::FNetworkDataTraits>(this);
+
+			if (NetworkPhysicsComponent->HasServerWorld())
 			{
-				// If the owner isn't a pawn, there's no chance of player input happening, so inputs to the PT are always produced on the server
-				NetworkPhysicsComponent->SetIsRelayingLocalInputs(true);
+				if (APawn* PawnOwner = Cast<APawn>(GetOwner()))
+				{
+					// When we're owned by a pawn, keep an eye on whether it's currently player-controlled or not
+					PawnOwner->ReceiveControllerChangedDelegate.AddUniqueDynamic(this, &ThisClass::HandleOwningPawnControllerChanged_Server);
+					HandleOwningPawnControllerChanged_Server(PawnOwner, nullptr, PawnOwner->Controller);
+				}
+				else
+				{
+					// If the owner isn't a pawn, there's no chance of player input happening, so inputs to the PT are always produced on the server
+					NetworkPhysicsComponent->SetIsRelayingLocalInputs(true);
+				}
 			}
 		}
-
-		bIsInputProducer = NetworkPhysicsComponent->IsLocallyControlled();
 	}
-
-	MoverBlackboard_Async = NewObject<UMoverBlackboard>(GetOwner(), TEXT("MoverBlackboard_Async"));
-	MovementMixer_Async = NewObject<UMovementMixer>(GetOwner(), TEXT("MovementMixer_Async"));
 }
 
 void UChaosMoverBackendComponent::UninitializeComponent()
@@ -90,17 +100,39 @@ void UChaosMoverBackendComponent::UninitializeComponent()
 		NetworkPhysicsComponent = nullptr;
 	}
 
-	if (MoverBlackboard_Async)
+	Super::UninitializeComponent();
+}
+
+void UChaosMoverBackendComponent::CreatePhysics()
+{
+	if (Chaos::FPhysicsSolver* Solver = GetPhysicsSolver())
 	{
-		MoverBlackboard_Async = nullptr;
+		if (Chaos::FPBDRigidParticle* Particle = GetControlledParticle())
+		{
+			if (Chaos::FSingleParticlePhysicsProxy* Proxy = static_cast<Chaos::FSingleParticlePhysicsProxy*>(Particle->GetProxy()))
+			{
+				CharacterGroundConstraint = MakeUnique<Chaos::FCharacterGroundConstraint>();
+				CharacterGroundConstraint->Init(Proxy);
+
+				Solver->RegisterObject(CharacterGroundConstraint.Get());
+
+				// Prevent the character particle from sleeping
+				Particle->SetSleepType(Chaos::ESleepType::NeverSleep);
+			}
+		}
 	}
+}
 
-	if (MovementMixer_Async)
+void UChaosMoverBackendComponent::DestroyPhysics()
+{
+	if (Chaos::FPhysicsSolver* Solver = GetPhysicsSolver())
 	{
-		MovementMixer_Async = nullptr;
+		if (CharacterGroundConstraint.IsValid())
+		{
+			// Note: Proxy gets destroyed when the constraint is deregistered and that deletes the constraint
+			Solver->UnregisterObject(CharacterGroundConstraint.Release());
+		}
 	}
-
-	Super::UninitializeComponent();
 }
 
 void UChaosMoverBackendComponent::HandleOwningPawnControllerChanged_Server(APawn* OwnerPawn, AController* OldController, AController* NewController)
@@ -109,7 +141,6 @@ void UChaosMoverBackendComponent::HandleOwningPawnControllerChanged_Server(APawn
 	if (NetworkPhysicsComponent)
 	{
 		NetworkPhysicsComponent->SetIsRelayingLocalInputs(!OwnerPawn->IsPlayerControlled());
-		bIsInputProducer = NetworkPhysicsComponent->IsLocallyControlled();
 	}
 }
 
@@ -117,12 +148,12 @@ void UChaosMoverBackendComponent::HandleUpdatedComponentPhysicsStateChanged(UPri
 {
 	if (StateChange == EComponentPhysicsStateChange::Destroyed)
 	{
+		bWantsDestroySim = true;
 	}
 	else if (StateChange == EComponentPhysicsStateChange::Created)
 	{
+		bWantsCreateSim = true;
 	}
-
-	InitSimulation();
 }
 
 Chaos::FPhysicsSolver* UChaosMoverBackendComponent::GetPhysicsSolver() const
@@ -155,17 +186,39 @@ Chaos::FPBDRigidParticle* UChaosMoverBackendComponent::GetControlledParticle() c
 	return nullptr;
 }
 
-TSharedPtr<UE::ChaosMover::FSimulation> UChaosMoverBackendComponent::CreateSimulation()
-{
-	return MakeShared<UE::ChaosMover::FSimulation>();
-}
-
 void UChaosMoverBackendComponent::InitSimulation()
 {
+	UMoverComponent& MoverComp = GetMoverComponent();
+
+	UChaosMoverSimulation::FInitParams Params;
+	for (const TPair<FName, TObjectPtr<UBaseMovementMode>>& Pair : MoverComp.MovementModes)
+	{
+		Params.ModesToRegister.Add(Pair.Key, TWeakObjectPtr<UBaseMovementMode>(Pair.Value.Get()));
+	}
+	for (const TObjectPtr<UBaseMovementModeTransition>& Transition : MoverComp.Transitions)
+	{
+		Params.TransitionsToRegister.Add(TWeakObjectPtr<UBaseMovementModeTransition>(Transition.Get()));
+	}
+	Params.MovementMixer = TWeakObjectPtr<UMovementMixer>(MoverComp.MovementMixer.Get());
+	Params.ImmediateModeTransition = TWeakObjectPtr<UImmediateMovementModeTransition>(ImmediateModeTransition.Get());
+	Params.NullMovementMode = TWeakObjectPtr<UNullMovementMode>(NullMovementMode.Get());
+	Params.StartingMovementMode = MoverComp.StartingMovementMode;
+	Params.CharacterConstraintProxy = CharacterGroundConstraint->GetProxy<Chaos::FCharacterGroundConstraintProxy>();
+	if (IPhysicsComponent* PhysicsComponent = Cast<IPhysicsComponent>(MoverComp.GetUpdatedComponent()))
+	{
+		if (Chaos::FPhysicsObjectHandle PhysObj = PhysicsComponent->GetPhysicsObjectByName(NAME_None))
+		{
+			Params.PhysicsObject = PhysObj;
+		}
+	}
+	Params.Solver = GetPhysicsSolver();
+	Params.World = GetWorld();
+	
 	SimOutputRecord.Clear();
 
 	UE::ChaosMover::FSimulationOutputData OutputData;
-	OutputData.SyncState.MovementMode = GetMoverComponent().StartingMovementMode;
+	FMoverAuxStateContext UnusedAuxState;
+	MoverComp.InitializeSimulationState(&OutputData.SyncState, &UnusedAuxState);
 
 	FMoverTimeStep TimeStep;
 	if (Chaos::FPhysicsSolver* Solver = GetPhysicsSolver())
@@ -177,92 +230,46 @@ void UChaosMoverBackendComponent::InitSimulation()
 
 	SimOutputRecord.Add(TimeStep, OutputData);
 
-	Simulation->SimBlackboardWeakPtr = TWeakObjectPtr<UMoverBlackboard>(MoverBlackboard_Async);
-	Simulation->MovementMixerWeakPtr =  TWeakObjectPtr<UMovementMixer>(MovementMixer_Async);
+	Params.InitialSyncState = OutputData.SyncState;
 
-	InitSimulationStateMachine();
-
-	Simulation->Init(OutputData.SyncState);
+	Simulation->Init(Params);
 }
 
-void UChaosMoverBackendComponent::InitSimulationStateMachine()
+void UChaosMoverBackendComponent::DeinitSimulation()
 {
-	if (!GetOwner())
-	{
-		return;
-	}
-
-	UMoverComponent& MoverComp = GetMoverComponent();
-
-	// Set up FSM and initial movement states
-
-	Simulation->MoverStateMachine.SetOwnerActorName(GetNameSafe(GetOwner()));
-	Simulation->MoverStateMachine.SetOwnerActorLocalNetRole(GetOwner()->GetLocalRole());
-
-	TObjectPtr<UImmediateMovementModeTransition> ImmediateMovementModeTransition = NewObject<UImmediateMovementModeTransition>(&MoverComp, TEXT("ImmediateMovementModeTransition"), RF_Transient);
-	TObjectPtr<UNullMovementMode> NullMovementMode = NewObject<UNullMovementMode>(&MoverComp, TEXT("NullMovementMode"), RF_Transient);
-
-	Simulation->MoverStateMachine.InitDefaultModesAndTransitions(ImmediateMovementModeTransition, NullMovementMode);
-
-	Simulation->MoverStateMachine.ClearAllMovementModes();
-	Simulation->MoverStateMachine.ClearAllGlobalTransitions();
-
-	bool bHasMatchingStartingState = false;
-
-	for (const TPair<FName, TObjectPtr<UBaseMovementMode>>& Element : MoverComp.MovementModes)
-	{
-		if (Element.Value.Get() == nullptr)
-		{
-			UE_LOG(LogMover, Warning, TEXT("Invalid Movement Mode type '%s' detected on %s. Mover actor will not function correctly."),
-				*Element.Key.ToString(), *GetNameSafe(GetOwner()));
-			continue;
-		}
-
-		Simulation->MoverStateMachine.RegisterMovementMode(Element.Key, Element.Value);
-
-		bHasMatchingStartingState |= (MoverComp.StartingMovementMode == Element.Key);
-	}
-
-	for (TObjectPtr<UBaseMovementModeTransition>& Transition : MoverComp.Transitions)
-	{
-		Simulation->MoverStateMachine.RegisterGlobalTransition(Transition);
-	}
-
-	if (!bHasMatchingStartingState)
-	{
-		UE_LOG(LogMover, Warning, TEXT("Invalid StartingMovementMode '%s' specified on %s. Mover actor will not function."),
-			*MoverComp.StartingMovementMode.ToString(), *GetNameSafe(GetOwner()));
-	}
-	else if (MoverComp.StartingMovementMode != NAME_None)
-	{
-		Simulation->MoverStateMachine.SetDefaultMode(MoverComp.StartingMovementMode);
-		Simulation->MoverStateMachine.QueueNextMode(MoverComp.StartingMovementMode);
-	}
+	Simulation->Deinit();
 }
 
 void UChaosMoverBackendComponent::BeginPlay()
 {
 	Super::BeginPlay();
 
-	Simulation = CreateSimulation();
-	InitSimulation();
-
-	// Register with the world subsystem
-	if (UChaosMoverSubsystem* ChaosMoverSubsystem = UWorld::GetSubsystem<UChaosMoverSubsystem>(GetWorld()))
+	if (UWorld* World = GetWorld(); World && World->IsGameWorld())
 	{
-		ChaosMoverSubsystem->Register(this);
-	}
+		// This is needed for now to create the mover component blackboard
+		GetMoverComponent().InitMoverSimulation();
 
-	// Register a callback to watch for component state changes
-	if (UPrimitiveComponent* PrimComp = Cast<UPrimitiveComponent>(GetMoverComponent().GetUpdatedComponent()))
-	{
-		PrimComp->OnComponentPhysicsStateChanged.AddUniqueDynamic(this, &ThisClass::HandleUpdatedComponentPhysicsStateChanged);
+		CreatePhysics();
+		InitSimulation();
+
+		// Register with the world subsystem
+		if (UChaosMoverSubsystem* ChaosMoverSubsystem = UWorld::GetSubsystem<UChaosMoverSubsystem>(GetWorld()))
+		{
+			ChaosMoverSubsystem->Register(this);
+		}
+
+		// Register a callback to watch for component state changes
+		if (UPrimitiveComponent* PrimComp = Cast<UPrimitiveComponent>(GetMoverComponent().GetUpdatedComponent()))
+		{
+			PrimComp->OnComponentPhysicsStateChanged.AddUniqueDynamic(this, &ThisClass::HandleUpdatedComponentPhysicsStateChanged);
+		}
 	}
 }
 
 void UChaosMoverBackendComponent::EndPlay(const EEndPlayReason::Type EndPlayReason)
 {
-	Simulation.Reset();
+	DeinitSimulation();
+	DestroyPhysics();
 
 	if (UPrimitiveComponent* PrimComp = Cast<UPrimitiveComponent>(GetMoverComponent().GetUpdatedComponent()))
 	{
@@ -295,9 +302,24 @@ void UChaosMoverBackendComponent::ProduceInputData(int32 PhysicsStep, int32 NumS
 {
 	Chaos::EnsureIsInGameThreadContext();
 
+	// Recreate the simulation if necessary
+	if (bWantsDestroySim)
+	{
+		DeinitSimulation();
+		DestroyPhysics();
+		bWantsDestroySim = false;
+		return;
+	}
+	if (bWantsCreateSim)
+	{
+		CreatePhysics();
+		InitSimulation();
+		bWantsCreateSim = false;
+	}
+
 	FMoverTimeStep TimeStep = GetMoverTimeStep(DeltaTimeInSeconds, NetworkPhysicsTickOffset);
 
-	if (bIsInputProducer)
+	if (!NetworkPhysicsComponent || NetworkPhysicsComponent->IsLocallyControlled())
 	{
 		GenerateInput(TimeStep, InputData);
 	}
@@ -309,7 +331,7 @@ void UChaosMoverBackendComponent::ProduceInputData(int32 PhysicsStep, int32 NumS
 	UMoverComponent& MoverComp = GetMoverComponent();
 
 	// Add default simulation input data
-	FChaosMoverSimulationDefaultInputs& SimInputs = InputData.InputParams.FindOrAddMutableDataByType<FChaosMoverSimulationDefaultInputs>();
+	FChaosMoverSimulationDefaultInputs& SimInputs = InputData.LocalInputParams.AuxStateCollection.FindOrAddMutableDataByType<FChaosMoverSimulationDefaultInputs>();
 	SimInputs.Gravity = MoverComp.GetGravityAcceleration();
 	SimInputs.UpDir = MoverComp.GetUpDirection();
 	SimInputs.OwningActor = GetOwner();
@@ -317,12 +339,15 @@ void UChaosMoverBackendComponent::ProduceInputData(int32 PhysicsStep, int32 NumS
 	{
 		PrimComp->InitSweepCollisionParams(SimInputs.CollisionQueryParams, SimInputs.CollisionResponseParams);
 		SimInputs.CollisionChannel = PrimComp->GetCollisionObjectType();
-		PrimComp->CalcBoundingCylinder(SimInputs.BoundingCylinderRadius, SimInputs.BoundingCylinderfHeight);
 	}
 	if (IPhysicsComponent* PhysComp = Cast<IPhysicsComponent>(MoverComp.GetUpdatedComponent()))
 	{
 		SimInputs.PhysicsObject = PhysComp->GetPhysicsObjectById(0); // Get the root physics object
 	}
+	if (const APhysicsVolume* CurPhysVolume = MoverComp.GetUpdatedComponent()->GetPhysicsVolume())
+	{
+		SimInputs.PhysicsObjectGravity = CurPhysVolume->GetGravityZ();
+	}
 
 	if (MoverComp.OnPreSimulationTick.IsBound())
 	{
@@ -347,16 +372,11 @@ void UChaosMoverBackendComponent::FinalizeFrame(double ResultsTimeInSeconds)
 {
 	Chaos::EnsureIsInGameThreadContext();
 
+	UMoverComponent& MoverComp = GetMoverComponent();
+
 	FMoverTimeStep TimeStep;
 	UE::ChaosMover::FSimulationOutputData InterpolatedOutput;
-	if (bIsUsingAsyncPhysics)
-	{
-		SimOutputRecord.GetInterpolated(ResultsTimeInSeconds, TimeStep, InterpolatedOutput);
-	}
-	else
-	{
-		InterpolatedOutput = SimOutputRecord.GetLatest();
-	}
+	SimOutputRecord.GetInterpolated(ResultsTimeInSeconds, TimeStep, InterpolatedOutput);
 
 	// Physics interactions in the last frame may have caused a change in position or velocity that's different from what a simple lerp would predict,
 	// so stomp the lerped sync state's transform data with that of the actual particle after the last sim frame
@@ -369,52 +389,25 @@ void UChaosMoverBackendComponent::FinalizeFrame(double ResultsTimeInSeconds)
 		TransformSyncState.MoveDirectionIntent = TransformSyncState.GetCapturedMovementBaseQuat().UnrotateVector(TransformSyncState.MoveDirectionIntent);
 	}
 
+	// Send the latest physics thread data to CVD
+	// Note we are sending the latest physics results, not the interpolated results
+#if WITH_CHAOS_VISUAL_DEBUGGER
+	if (FChaosVisualDebuggerTrace::IsTracing())
+	{
+		const UE::ChaosMover::FSimulationOutputData& LatestOutput = SimOutputRecord.GetLatest();
+		FMoverCVDRuntimeTrace::TraceMoverData(&MoverComp, &InterpolatedOutput.LastUsedInputCmd, &InterpolatedOutput.SyncState);
+	}
+#endif
+
 	// This eventually will process any events queued by the simulation during the async tick.
 	// For now it's just movement mode activation, deactivation and transition callbacks.
 	FlushModeChangeEvents();
 
-	UMoverComponent& MoverComp = GetMoverComponent();
 	MoverComp.SetSimulationOutput(TimeStep, InterpolatedOutput);
-	MoverComp.OnPostSimulationTick.Broadcast(TimeStep);
-}
-
-void UChaosMoverBackendComponent::ApplyNetInputData(const FMoverInputCmdContext& InNetInputCmd)
-{
-	if (Simulation)
-	{
-		Simulation->ApplyNetInputData(InNetInputCmd);
-	}
-}
-
-void UChaosMoverBackendComponent::BuildNetInputData(FMoverInputCmdContext& OutNetInputCmd) const
-{
-	if (Simulation)
-	{
-		Simulation->BuildNetInputData(OutNetInputCmd);
-	}
-}
-
-void UChaosMoverBackendComponent::ApplyNetStateData(const FMoverSyncState& InNetSyncState)
-{
-	if (Simulation)
-	{
-		Simulation->ApplyNetStateData(InNetSyncState);
-	}
-}
-
-void UChaosMoverBackendComponent::BuildNetStateData(FMoverSyncState& OutNetSyncState) const
-{
-	if (Simulation)
+	
+	if (MoverComp.OnPostSimulationTick.IsBound())
 	{
-		Simulation->BuildNetStateData(OutNetSyncState);
-	}
-}
-
-void UChaosMoverBackendComponent::RegisterSimulation() const
-{
-	if (UChaosMoverSubsystem* ChaosMoverSubsystem = UWorld::GetSubsystem<UChaosMoverSubsystem>(GetWorld()))
-	{
-		ChaosMoverSubsystem->RegisterSimulation(Simulation);
+		MoverComp.OnPostSimulationTick.Broadcast(TimeStep);
 	}
 }
 
@@ -423,32 +416,29 @@ void UChaosMoverBackendComponent::RegisterSimulation() const
 */
 void UChaosMoverBackendComponent::FlushModeChangeEvents()
 {
-	if (Simulation)
+	const TArray<UE::ChaosMover::FMovementModeChangeEvent>& QueuedModeChangeEvents = Simulation->GetQueuedModeChangeEvents();
+	UMoverComponent& MoverComp = GetMoverComponent();
+
+	for (int Idx = QueuedModeChangeEvents.Num() - 1; Idx >= 0; --Idx)
 	{
-		const TArray<UE::ChaosMover::FMovementModeChangeEvent>& QueuedModeChangeEvents = Simulation->MoverStateMachine.GetQueuedModeChangeEvents();
-		UMoverComponent& MoverComp = GetMoverComponent();
+		const UE::ChaosMover::FMovementModeChangeEvent& MovementModeChangeEvent = QueuedModeChangeEvents[Idx];
+		TStrongObjectPtr<UBaseMovementMode> PreviousMovementMode = MovementModeChangeEvent.PreviousMovementMode.Pin();
+		TStrongObjectPtr<UBaseMovementMode> NextMovementMode = MovementModeChangeEvent.NextMovementMode.Pin();
 
-		for (int i = QueuedModeChangeEvents.Num() - 1; i >= 0; --i)
+		if (PreviousMovementMode)
 		{
-			const UE::ChaosMover::FMovementModeChangeEvent& MovementModeChangeEvent = QueuedModeChangeEvents[i];
-			TStrongObjectPtr<UBaseMovementMode> PreviousMovementMode = MovementModeChangeEvent.PreviousMovementMode.Pin();
-			TStrongObjectPtr<UBaseMovementMode> NextMovementMode = MovementModeChangeEvent.NextMovementMode.Pin();
-
-			if (PreviousMovementMode)
-			{
-				PreviousMovementMode->Deactivate();
-			}
-
-			if (NextMovementMode)
-			{
-				NextMovementMode->Activate();
-			}
+			PreviousMovementMode->Deactivate();
+		}
 
-			// Signal movement mode change event
-			MoverComp.OnMovementModeChanged.Broadcast(MovementModeChangeEvent.PreviousMovementModeName, MovementModeChangeEvent.NextMovementModeName);
+		if (NextMovementMode)
+		{
+			NextMovementMode->Activate();
 		}
 
-		Simulation->MoverStateMachine.ClearQueuedModeChangeEvents();
+		// Signal movement mode change event
+		MoverComp.OnMovementModeChanged.Broadcast(MovementModeChangeEvent.PreviousMovementModeName, MovementModeChangeEvent.NextMovementModeName);
 	}
+
+	Simulation->ClearQueuedModeChangeEvents();
 }
 
diff --git a/Engine/Plugins/Experimental/ChaosMover/Source/ChaosMover/Private/Backends/ChaosMoverSubsystem.cpp b/Engine/Plugins/Experimental/ChaosMover/Source/ChaosMover/Private/Backends/ChaosMoverSubsystem.cpp
--- a/Engine/Plugins/Experimental/ChaosMover/Source/ChaosMover/Private/Backends/ChaosMoverSubsystem.cpp
+++ b/Engine/Plugins/Experimental/ChaosMover/Source/ChaosMover/Private/Backends/ChaosMoverSubsystem.cpp
@@ -29,7 +29,6 @@ void UChaosMoverSubsystem::OnWorldBeginPlay(UWorld& InWorld)
 		if (Chaos::FPhysicsSolver* Solver = PhysScene->GetSolver())
 		{
 			AsyncCallback = Solver->CreateAndRegisterSimCallbackObject_External<UE::ChaosMover::FAsyncCallback>();
-			AsyncCallback->PhysicsSolver = Solver;
 
 			if (FNetworkPhysicsCallback* SolverCallback = static_cast<FNetworkPhysicsCallback*>(Solver->GetRewindCallback()))
 			{
@@ -74,23 +73,11 @@ void UChaosMoverSubsystem::Deinitialize()
 void UChaosMoverSubsystem::Register(TWeakObjectPtr<UChaosMoverBackendComponent> InBackend)
 {
 	Backends.AddUnique(InBackend);
-	bBackendsListDirty = true;
 }
 
 void UChaosMoverSubsystem::Unregister(TWeakObjectPtr<UChaosMoverBackendComponent> InBackend)
 {
 	Backends.Remove(InBackend);
-	bBackendsListDirty = true;
-}
-
-void UChaosMoverSubsystem::RegisterSimulation(TSharedPtr<UE::ChaosMover::FSimulation> Simulation)
-{
-	AsyncCallback->Simulations.Add(Simulation);
-}
-
-void UChaosMoverSubsystem::UnregisterAllSimulations()
-{
-	AsyncCallback->Simulations.Empty();
 }
 
 void UChaosMoverSubsystem::InjectInputs_External(int32 PhysicsStep, int32 NumSteps)
@@ -101,37 +88,32 @@ void UChaosMoverSubsystem::InjectInputs_External(int32 PhysicsStep, int32 NumSte
 	}
 
 	// Clear invalid data before we attempt to use it.
-	bBackendsListDirty |= Backends.RemoveAllSwap([](const TWeakObjectPtr<UChaosMoverBackendComponent> Backend) {
-		return (Backend.IsValid() == false);
-		}) > 0;
+	Backends.RemoveAllSwap([](const TWeakObjectPtr<UChaosMoverBackendComponent> Backend) {
+		return !Backend.IsValid();
+		});
+
+	UE::ChaosMover::FAsyncCallbackInput* AsyncInput = AsyncCallback->GetProducerInputData_External();
+	AsyncInput->Reset();
+	AsyncInput->InputData.SetNum(Backends.Num());
 
-	// Copy a list of active backends to the callback
-	if (bBackendsListDirty)
+	for (TWeakObjectPtr<UChaosMoverBackendComponent>& Backend : Backends)
 	{
-		AsyncCallback->Backends.Empty();
-		for (int Idx = 0; Idx < Backends.Num(); ++Idx)
-		{
-			Backends[Idx]->RegisterSimulation();
-			AsyncCallback->Backends.Add(Backends[Idx]);
-		}
-		bBackendsListDirty = false;
+		AsyncInput->Backends.Add(Backend);
 	}
 
-	// Refresh the network physics tick offset in the async callback
+	// Refresh the network physics tick offset in the async input
 	const int32 NetPhysicsTickOffset = GetNetworkPhysicsTickOffset();
-	AsyncCallback->NetworkPhysicsTickOffset = NetPhysicsTickOffset;
-
-	UE::ChaosMover::FAsyncCallbackInput* AsyncInput = AsyncCallback->GetProducerInputData_External();
-	AsyncInput->Reset();
-	AsyncInput->InputData.SetNum(Backends.Num());
+	AsyncInput->NetworkPhysicsTickOffset = NetPhysicsTickOffset;
 
 	double DeltaSeconds = 0.0;
-	if (const UWorld* World = GetWorld())
+	if (UWorld* World = GetWorld())
 	{
-		if (const FPhysScene* Scene = World->GetPhysicsScene())
+		if (FPhysScene* Scene = World->GetPhysicsScene())
 		{
-			if (const Chaos::FPhysicsSolver* Solver = Scene->GetSolver())
+			if (Chaos::FPhysicsSolver* Solver = Scene->GetSolver())
 			{
+				AsyncInput->PhysicsSolver = Solver;
+
 				if (Solver->IsUsingAsyncResults())
 				{
 					DeltaSeconds = Solver->GetAsyncDeltaTime();
@@ -164,14 +146,14 @@ void UChaosMoverSubsystem::OnPostPhysicsTick(FChaosScene* Scene)
 	// Pop each async output
 	while (Chaos::TSimCallbackOutputHandle<UE::ChaosMover::FAsyncCallbackOutput> AsyncOutput = AsyncCallback->PopFutureOutputData_External())
 	{
-		auto LambdaParallelUpdate = [&AsyncOutput, this](int32 Idx) {
-			if (UChaosMoverBackendComponent* Backend = AsyncCallback->Backends[Idx].Get())
+		auto LambdaParallelUpdate = [&AsyncOutput](int32 Idx) {
+			if (UChaosMoverBackendComponent* Backend = AsyncOutput->Backends[Idx].Get())
 			{
 				Backend->ConsumeOutputData(AsyncOutput->TimeStep[Idx], AsyncOutput->OutputData[Idx]);
 			}
 		};
 
-		Chaos::PhysicsParallelFor(AsyncCallback->Backends.Num(), LambdaParallelUpdate, ForceSingleThread);
+		Chaos::PhysicsParallelFor(AsyncOutput->Backends.Num(), LambdaParallelUpdate, ForceSingleThread);
 	}
 
 	double ResultsTime = 0.0;
@@ -184,14 +166,14 @@ void UChaosMoverSubsystem::OnPostPhysicsTick(FChaosScene* Scene)
 	}
 
 	// Call finalize frame
-	auto LambdaParallelUpdate = [ResultsTime, this](int32 Idx) {
-		if (UChaosMoverBackendComponent* Backend = AsyncCallback->Backends[Idx].Get())
+	auto LambdaParallelUpdate = [this, &ResultsTime](int32 Idx) {
+		if (UChaosMoverBackendComponent* Backend = Backends[Idx].Get())
 		{
 			Backend->FinalizeFrame(ResultsTime);
 		}
 	};
 
-	Chaos::PhysicsParallelFor(AsyncCallback->Backends.Num(), LambdaParallelUpdate, ForceSingleThread);
+	Chaos::PhysicsParallelFor(Backends.Num(), LambdaParallelUpdate, ForceSingleThread);
 }
 
 int32 UChaosMoverSubsystem::GetNetworkPhysicsTickOffset() const
diff --git a/Engine/Plugins/Experimental/ChaosMover/Source/ChaosMover/Private/Backends/ChaosMoverSubsystem.h b/Engine/Plugins/Experimental/ChaosMover/Source/ChaosMover/Private/Backends/ChaosMoverSubsystem.h
--- a/Engine/Plugins/Experimental/ChaosMover/Source/ChaosMover/Private/Backends/ChaosMoverSubsystem.h
+++ b/Engine/Plugins/Experimental/ChaosMover/Source/ChaosMover/Private/Backends/ChaosMoverSubsystem.h
@@ -2,8 +2,8 @@
 
 #pragma once
 
-#include "Subsystems/WorldSubsystem.h"
 #include "Chaos/Declares.h"
+#include "Subsystems/WorldSubsystem.h"
 
 #include "ChaosMoverSubsystem.generated.h"
 
@@ -25,12 +25,6 @@ public:
 	void Register(TWeakObjectPtr<UChaosMoverBackendComponent> Backend);
 	void Unregister(TWeakObjectPtr<UChaosMoverBackendComponent> Backend);
 
-	/**
-	 * Registers a backend simulation's object shared pointer so it is kept alive during simulation,
-	 * even if its owner actor gets destroyed and garbage collected
-	 */ 
-	void RegisterSimulation(TSharedPtr<UE::ChaosMover::FSimulation> Simulation);
-
 protected:
 	virtual bool DoesSupportWorldType(const EWorldType::Type WorldType) const override;
 	virtual void OnWorldBeginPlay(UWorld& InWorld) override;
@@ -41,14 +35,7 @@ protected:
 	void InjectInputs_External(int32 PhysicsStep, int32 NumSteps);
 	void OnPostPhysicsTick(FChaosScene* Scene);
 
-	/**
-	 * Unregisters all simulations. This is called internally when the list of backends is refreshed
-	 * The existing backends then all register their simulations again so they are kept alive
-	 */
-	void UnregisterAllSimulations();
-
 	TArray<TWeakObjectPtr<UChaosMoverBackendComponent>> Backends;
-	bool bBackendsListDirty = true;
 
 	FDelegateHandle InjectInputsExternalCallbackHandle;
 	FDelegateHandle PhysScenePostTickCallbackHandle;
diff --git a/Engine/Plugins/Experimental/ChaosMover/Source/ChaosMover/Private/Backends/NetworkChaosMoverData.cpp b/Engine/Plugins/Experimental/ChaosMover/Source/ChaosMover/Private/Backends/NetworkChaosMoverData.cpp
--- a/Engine/Plugins/Experimental/ChaosMover/Source/ChaosMover/Private/Backends/NetworkChaosMoverData.cpp
+++ b/Engine/Plugins/Experimental/ChaosMover/Source/ChaosMover/Private/Backends/NetworkChaosMoverData.cpp
@@ -2,14 +2,14 @@
 
 #include "NetworkChaosMoverData.h"
 
-#include "ChaosMover/ChaosMoverSimulation.h"
 #include "ChaosMover/Backends/ChaosMoverBackend.h"
+#include "ChaosMover/ChaosMoverSimulation.h"
 
 void FNetworkChaosMoverInputData::ApplyData(UActorComponent* NetworkComponent) const
 {
 	if (UChaosMoverBackendComponent* BackendComp = Cast<UChaosMoverBackendComponent>(NetworkComponent))
 	{
-		BackendComp->ApplyNetInputData(InputCmdContext);
+		BackendComp->GetSimulation()->ApplyNetInputData(InputCmdContext);
 	}
 }
 
@@ -17,7 +17,7 @@ void FNetworkChaosMoverInputData::BuildData(const UActorComponent* NetworkCompon
 {
 	if (const UChaosMoverBackendComponent* BackendComp = Cast<const UChaosMoverBackendComponent>(NetworkComponent))
 	{
-		BackendComp->BuildNetInputData(InputCmdContext);
+		BackendComp->GetSimulation()->BuildNetInputData(InputCmdContext);
 	}
 }
 
@@ -84,7 +84,7 @@ void FNetworkChaosMoverStateData::ApplyData(UActorComponent* NetworkComponent) c
 	{
 		if (UChaosMoverBackendComponent* BackendComp = Cast<UChaosMoverBackendComponent>(NetworkComponent))
 		{
-			BackendComp->ApplyNetStateData(SyncState);
+			BackendComp->GetSimulation()->ApplyNetStateData(SyncState);
 		}
 	}
 }
@@ -95,7 +95,7 @@ void FNetworkChaosMoverStateData::BuildData(const UActorComponent* NetworkCompon
 	{
 		if (const UChaosMoverBackendComponent* BackendComp = Cast<const UChaosMoverBackendComponent>(NetworkComponent))
 		{
-			BackendComp->BuildNetStateData(SyncState);
+			BackendComp->GetSimulation()->BuildNetStateData(SyncState);
 		}
 	}
 }
diff --git a/Engine/Plugins/Experimental/ChaosMover/Source/ChaosMover/Private/ChaosMovementMode.cpp b/Engine/Plugins/Experimental/ChaosMover/Source/ChaosMover/Private/ChaosMovementMode.cpp
new file mode 100644
--- /dev/null
+++ b/Engine/Plugins/Experimental/ChaosMover/Source/ChaosMover/Private/ChaosMovementMode.cpp
@@ -0,0 +1,26 @@
+// Copyright Epic Games, Inc. All Rights Reserved.
+
+#include "ChaosMover/ChaosMovementMode.h"
+
+#include "ChaosMover/ChaosMovementModeTransition.h"
+#include "ChaosMover/ChaosMoverLog.h"
+
+#include UE_INLINE_GENERATED_CPP_BY_NAME(ChaosMovementMode)
+
+UChaosMovementMode::UChaosMovementMode(const FObjectInitializer& ObjectInitializer)
+	: Super(ObjectInitializer)
+	, Simulation(nullptr)
+{
+}
+
+void UChaosMovementMode::SetSimulation(UChaosMoverSimulation* InSimulation)
+{
+	Simulation = InSimulation;
+	for (TObjectPtr<UBaseMovementModeTransition>& Transition : Transitions)
+	{
+		if (UChaosMovementModeTransition* ChaosTransition = Cast<UChaosMovementModeTransition>(Transition.Get()))
+		{
+			ChaosTransition->SetSimulation(InSimulation);
+		}
+	}
+}
\ No newline at end of file
diff --git a/Engine/Plugins/Experimental/ChaosMover/Source/ChaosMover/Private/ChaosMovementModeTransition.cpp b/Engine/Plugins/Experimental/ChaosMover/Source/ChaosMover/Private/ChaosMovementModeTransition.cpp
new file mode 100644
--- /dev/null
+++ b/Engine/Plugins/Experimental/ChaosMover/Source/ChaosMover/Private/ChaosMovementModeTransition.cpp
@@ -0,0 +1,13 @@
+// Copyright Epic Games, Inc. All Rights Reserved.
+
+#include "ChaosMover/ChaosMovementModeTransition.h"
+
+#include "ChaosMover/ChaosMoverLog.h"
+
+#include UE_INLINE_GENERATED_CPP_BY_NAME(ChaosMovementModeTransition)
+
+UChaosMovementModeTransition::UChaosMovementModeTransition(const FObjectInitializer& ObjectInitializer)
+	: Super(ObjectInitializer)
+	, Simulation(nullptr)
+{
+}
\ No newline at end of file
diff --git a/Engine/Plugins/Experimental/ChaosMover/Source/ChaosMover/Private/ChaosMoverDeveloperSettings.cpp b/Engine/Plugins/Experimental/ChaosMover/Source/ChaosMover/Private/ChaosMoverDeveloperSettings.cpp
--- a/Engine/Plugins/Experimental/ChaosMover/Source/ChaosMover/Private/ChaosMoverDeveloperSettings.cpp
+++ b/Engine/Plugins/Experimental/ChaosMover/Source/ChaosMover/Private/ChaosMoverDeveloperSettings.cpp
@@ -19,5 +19,10 @@ namespace UE::ChaosMover
 		bool bDrawGroundQueries = false;
 		FAutoConsoleVariableRef CVarChaosMoverDrawGroundQueries(TEXT("ChaosMover.DebugDraw.GroundQueries"),
 			bDrawGroundQueries, TEXT("Draw ground queries."));
+
+		bool bSkipGenerateMoveIfOverridden = true;
+		FAutoConsoleVariableRef CVarChaosMoverSkipGenerateMoveIfOverridden(TEXT("ChaosMover.Perf.SkipGenerateMoveIfOverridden"),
+			bSkipGenerateMoveIfOverridden, TEXT("If true and we have a layered move fully overriding movement, then we will skip calling OnGenerateMove on the active movement mode for better performance\n")
+		);
 	}
 }
\ No newline at end of file
diff --git a/Engine/Plugins/Experimental/ChaosMover/Source/ChaosMover/Private/ChaosMoverSimulation.cpp b/Engine/Plugins/Experimental/ChaosMover/Source/ChaosMover/Private/ChaosMoverSimulation.cpp
--- a/Engine/Plugins/Experimental/ChaosMover/Source/ChaosMover/Private/ChaosMoverSimulation.cpp
+++ b/Engine/Plugins/Experimental/ChaosMover/Source/ChaosMover/Private/ChaosMoverSimulation.cpp
@@ -2,116 +2,436 @@
 
 #include "ChaosMover/ChaosMoverSimulation.h"
 
+#include "Chaos/Character/CharacterGroundConstraint.h"
+#include "Chaos/Character/CharacterGroundConstraintContainer.h"
+#include "Chaos/PhysicsObjectInternalInterface.h"
+#include "ChaosMover/ChaosMovementMode.h"
+#include "ChaosMover/ChaosMovementModeTransition.h"
 #include "ChaosMover/ChaosMoverLog.h"
 #include "ChaosMover/ChaosMoverSimulationTypes.h"
-#include "MoverSimulationTypes.h"
+#include "ChaosMover/Utilities/ChaosMoverQueryUtils.h"
 #include "Framework/Threading.h"
-#include "PBDRigidsSolver.h"
-#include "MovementModeStateMachine.h"
+#include "MoveLibrary/FloorQueryUtils.h"
 #include "MoveLibrary/MovementMixer.h"
 #include "MoveLibrary/MoverBlackboard.h"
+#include "MoveLibrary/WaterMovementUtils.h"
+#include "MovementModeStateMachine.h"
+#include "MoverSimulationTypes.h"
+#include "PBDRigidsSolver.h"
+#include "PhysicsProxy/CharacterGroundConstraintProxy.h"
 #include "UObject/UObjectGlobals.h"
-#include "MoverDeveloperSettings.h"
 
-namespace UE::ChaosMover
+#include UE_INLINE_GENERATED_CPP_BY_NAME(ChaosMoverSimulation)
+
+UChaosMoverSimulation::UChaosMoverSimulation()
+{
+	Blackboard = CreateDefaultSubobject<UMoverBlackboard>(TEXT("ChaosMoverSimulationBlackboard"));
+}
+
+const UBaseMovementMode* UChaosMoverSimulation::GetCurrentMovementMode() const
+{
+	return StateMachine.GetCurrentMode().Get();
+}
+
+const UBaseMovementMode* UChaosMoverSimulation::FindMovementModeByName(const FName& Name) const
+{
+	return StateMachine.FindMovementMode(Name).Get();
+}
+
+void UChaosMoverSimulation::InitNetInputData(const FMoverInputCmdContext& InNetInputCmd)
+{
+	InputCmd = InNetInputCmd;
+}
+
+void UChaosMoverSimulation::ApplyNetInputData(const FMoverInputCmdContext& InNetInputCmd)
 {
-	FSimulation::~FSimulation()
+	InputCmd = InNetInputCmd;
+	bInputCmdOverridden = true;
+}
+
+void UChaosMoverSimulation::BuildNetInputData(FMoverInputCmdContext& OutNetInputCmd) const
+{
+	OutNetInputCmd = InputCmd;
+}
+
+void UChaosMoverSimulation::ApplyNetStateData(const FMoverSyncState& InNetSyncState)
+{
+	CurrentSyncState = InNetSyncState;
+}
+
+void UChaosMoverSimulation::BuildNetStateData(FMoverSyncState& OutNetSyncState) const
+{
+	OutNetSyncState = CurrentSyncState;
+}
+
+void UChaosMoverSimulation::Init(const FInitParams& InitParams)
+{
+	MovementMixerWeakPtr = InitParams.MovementMixer.Get();
+	CharacterConstraintProxy = InitParams.CharacterConstraintProxy;
+	PhysicsObject = InitParams.PhysicsObject;
+	Solver = InitParams.Solver;
+	World = InitParams.World;
+
+	CurrentSyncState = InitParams.InitialSyncState;
+
+	StateMachine.InitDefaultModesAndTransitions(InitParams.ImmediateModeTransition, InitParams.NullMovementMode);
+	StateMachine.ClearAllMovementModes();
+	StateMachine.ClearAllGlobalTransitions();
+
+	for (const TPair<FName, TWeakObjectPtr<UBaseMovementMode>>& Element : InitParams.ModesToRegister)
 	{
+		if (Element.Value.Get() == nullptr)
+		{
+			UE_LOG(LogChaosMover, Warning, TEXT("Invalid Movement Mode type '%s' detected. Mover actor will not function correctly."), *Element.Key.ToString());
+			continue;
+		}
+
+		if (UChaosMovementMode* ChaosMode = Cast<UChaosMovementMode>(Element.Value.Get()))
+		{
+			ChaosMode->SetSimulation(this);
+		}
+
+		StateMachine.RegisterMovementMode(Element.Key, Element.Value);
 	}
-	
-	void FSimulation::InitNetInputData(const FMoverInputCmdContext& InNetInputCmd)
+
+	if (InitParams.StartingMovementMode != NAME_None)
 	{
-		InputCmd = InNetInputCmd;
+		StateMachine.SetDefaultMode(InitParams.StartingMovementMode);
+		StateMachine.SetModeImmediately(InitParams.StartingMovementMode);
 	}
 
-	void FSimulation::ApplyNetInputData(const FMoverInputCmdContext& InNetInputCmd)
+	for (const TWeakObjectPtr<UBaseMovementModeTransition>& Transition : InitParams.TransitionsToRegister)
 	{
-		InputCmd = InNetInputCmd;
-		bInputCmdOverridden = true;
+		if (UChaosMovementModeTransition* ChaosTransition= Cast<UChaosMovementModeTransition>(Transition.Get()))
+		{
+			ChaosTransition->SetSimulation(this);
+		}
+
+		StateMachine.RegisterGlobalTransition(Transition);
 	}
 
-	void FSimulation::BuildNetInputData(FMoverInputCmdContext& OutNetInputCmd) const
+	OnInit();
+}
+
+void UChaosMoverSimulation::Deinit()
+{
+	OnDeinit();
+}
+
+void UChaosMoverSimulation::OnInit()
+{
+	
+}
+
+void UChaosMoverSimulation::OnDeinit()
+{
+
+}
+
+void UChaosMoverSimulation::SimulationTick(const FMoverTimeStep& TimeStep, const UE::ChaosMover::FSimulationInputData& InputData, UE::ChaosMover::FSimulationOutputData& OutputData)
+{
+	Chaos::EnsureIsInPhysicsThreadContext();
+
+	OnPreSimulationTick(TimeStep, InputData);
+	OnSimulationTick(TimeStep, InputData, OutputData);
+	OnPostSimulationTick(TimeStep, OutputData);
+}
+
+void UChaosMoverSimulation::ModifyContacts(Chaos::FCollisionContactModifier& Modifier) const
+{
+	Chaos::EnsureIsInPhysicsThreadContext();
+
+	OnModifyContacts(Modifier);
+}
+
+void UChaosMoverSimulation::PreSimulationTickCharacter(const IChaosCharacterMovementModeInterface& CharacterMode, const FMoverTimeStep& TimeStep, const UE::ChaosMover::FSimulationInputData& InputData)
+{
+	const FMoverDefaultSyncState* DefaultSyncState = CurrentSyncState.SyncStateCollection.FindDataByType<FMoverDefaultSyncState>();
+	check(DefaultSyncState);
+	check(Blackboard.Get());
+
+	if (const FChaosMoverSimulationDefaultInputs* DefaultInputs = InputData.LocalInputParams.AuxStateCollection.FindDataByType<FChaosMoverSimulationDefaultInputs>())
 	{
-		OutNetInputCmd = InputCmd;
+		UE::ChaosMover::Utils::FFloorSweepParams SweepParams{
+			.ResponseParams = DefaultInputs->CollisionResponseParams,
+			.QueryParams = DefaultInputs->CollisionQueryParams,
+			.Location = DefaultSyncState->GetLocation_WorldSpace(),
+			.DeltaPos = DefaultSyncState->GetVelocity_WorldSpace() * TimeStep.StepMs * 0.001f,
+			.UpDir = DefaultInputs->UpDir,
+			.World = World,
+			.QueryDistance = 1.2f * CharacterMode.GetTargetHeight(),
+			.QueryRadius = CharacterMode.GetGroundQueryRadius(),
+			.MaxWalkSlopeCosine = CharacterMode.GetMaxWalkSlopeCosine(),
+			.TargetHeight = CharacterMode.GetTargetHeight(),
+			.CollisionChannel = DefaultInputs->CollisionChannel
+		};
+
+		FFloorCheckResult FloorResult;
+		FWaterCheckResult WaterResult;
+		UE::ChaosMover::Utils::FloorSweep_Internal(SweepParams, FloorResult, WaterResult);
+
+		Blackboard->Set(CommonBlackboard::LastFloorResult, FloorResult);
+		Blackboard->Set(CommonBlackboard::LastWaterResult, WaterResult);
 	}
+}
 
-	void FSimulation::ApplyNetStateData(const FMoverSyncState& InNetSyncState)
+void UChaosMoverSimulation::OnPreSimulationTick(const FMoverTimeStep& TimeStep, const UE::ChaosMover::FSimulationInputData& InputData)
+{
+	if (bInputCmdOverridden)
 	{
-		SyncState = InNetSyncState;
+		InputData.InputCmd = InputCmd;
 	}
 
-	void FSimulation::BuildNetStateData(FMoverSyncState& OutNetSyncState) const
+	if (TimeStep.bIsResimulating)
 	{
-		OutNetSyncState = SyncState;
+		check(Solver);
+		if (Solver->GetEvolution()->IsResetting())
+		{
+			// Rollback blackboard on the first frame of resimulation
+			Blackboard->Invalidate(EInvalidationReason::Rollback);
+		}
 	}
 
-	void FSimulation::Init(const FMoverSyncState& InitialSyncState)
+	// Update the sync state from the current physics state
+	FMoverDefaultSyncState& DefaultSyncState = CurrentSyncState.SyncStateCollection.FindOrAddMutableDataByType<FMoverDefaultSyncState>();
+	if (const Chaos::FPBDRigidParticleHandle* ParticleHandle = UChaosMoverSimulation::GetControlledParticle())
 	{
-		SyncState = InitialSyncState;
-		OnInit();
+		DefaultSyncState.SetTransforms_WorldSpace(ParticleHandle->GetX(), FRotator(ParticleHandle->GetR()), ParticleHandle->GetV());
 	}
 
-	void FSimulation::OnInit()
+	// TODO - make this more extensible
+	if (TStrongObjectPtr<UBaseMovementMode> CurrentModePtr = StateMachine.GetCurrentMode().Pin())
 	{
-		
+		if (IChaosCharacterMovementModeInterface* CharacterMode = Cast<IChaosCharacterMovementModeInterface>(CurrentModePtr.Get()))
+		{
+			PreSimulationTickCharacter(*CharacterMode, TimeStep, InputData);
+		}
 	}
+}
 
-	void FSimulation::SimulationTick(const FMoverTimeStep& TimeStep, const FSimulationInputData& InputData, FSimulationOutputData& OutputData)
+void UChaosMoverSimulation::OnSimulationTick(const FMoverTimeStep& TimeStep, const UE::ChaosMover::FSimulationInputData& InputData, UE::ChaosMover::FSimulationOutputData& OutputData)
+{
+	check(Blackboard.Get());
+
+	FMoverTickStartData TickStartData(InputData.InputCmd, CurrentSyncState, InputData.LocalInputParams);
+	FMoverTickEndData TickEndData(&CurrentSyncState, &InputData.LocalInputParams);
+	StateMachine.OnSimulationTick(TimeStep, TickStartData, Blackboard.Get(), MovementMixerWeakPtr.Get(), TickEndData);
+
+	// Copy the sync state locally and to the output data
+	OutputData.SyncState = TickEndData.SyncState;
+	OutputData.LastUsedInputCmd = InputData.InputCmd;
+
+	// Store the delta position and rotation locally
+	if (Chaos::FPBDRigidParticleHandle* ParticleHandle = UChaosMoverSimulation::GetControlledParticle())
+	{
+		const FMoverDefaultSyncState* DefaultSyncState = OutputData.SyncState.SyncStateCollection.FindDataByType<FMoverDefaultSyncState>();
+		check(DefaultSyncState);
+
+		FRotator TmpRotation = DefaultSyncState->GetOrientation_WorldSpace() - FRotator(ParticleHandle->GetR());
+		FRotator Winding;
+		TmpRotation.GetWindingAndRemainder(Winding, DeltaRotation);
+
+		DeltaPosition = DefaultSyncState->GetLocation_WorldSpace() - ParticleHandle->GetX();
+	}
+	else
 	{
-		Chaos::EnsureIsInPhysicsThreadContext();
+		DeltaPosition = FVector::ZeroVector;
+		DeltaRotation = FRotator::ZeroRotator;
+	}
+}
+
+void UChaosMoverSimulation::PostSimulationTickCharacter(const IChaosCharacterMovementModeInterface& CharacterMode, const FMoverTimeStep& TimeStep, UE::ChaosMover::FSimulationOutputData& OutputData)
+{
+	const FMoverDefaultSyncState* DefaultSyncState = OutputData.SyncState.SyncStateCollection.FindDataByType<FMoverDefaultSyncState>();
+	check(DefaultSyncState);
+
+	// Set the particle velocity to hit the target sync state transform
+	if (Chaos::FPBDRigidParticleHandle* ParticleHandle = UChaosMoverSimulation::GetControlledParticle())
+	{
+		ParticleHandle->SetV(DefaultSyncState->GetVelocity_WorldSpace());
+
+		float TargetDeltaFacing = FMath::DegreesToRadians(DeltaRotation.Yaw);
+		const float DeltaTimeSeconds = TimeStep.StepMs * 0.001f;
+		if (DeltaTimeSeconds > UE_SMALL_NUMBER)
+		{
+			FVector AngularVelocity = (TargetDeltaFacing / DeltaTimeSeconds) * FVector::UpVector;
 
-		OnPreSimulationTick(TimeStep, InputData);
-		OnSimulationTick(TimeStep, InputData, OutputData);
-		OnPostSimulationTick(TimeStep, OutputData);
+			if (!CharacterMode.ShouldCharacterRemainUpright())
+			{
+				FVector PreviousUp = ParticleHandle->GetR().RotateVector(FVector::UpVector).GetSafeNormal();
+				FVector TargetUp = DefaultSyncState->GetOrientation_WorldSpace().Quaternion().RotateVector(FVector::UpVector).GetSafeNormal();
+				AngularVelocity += (PreviousUp.Cross(TargetUp).GetSafeNormal() * FMath::Acos(TargetUp.Dot(PreviousUp))) / DeltaTimeSeconds;
+			}
+
+			ParticleHandle->SetW(AngularVelocity);
+		}
 	}
+}
 
-	void FSimulation::ModifyContacts(Chaos::FCollisionContactModifier& Modifier)
+void UChaosMoverSimulation::PostSimulationTickCharacterConstraint(const IChaosCharacterConstraintMovementModeInterface& CharacterConstraintMode, const FMoverTimeStep& TimeStep, UE::ChaosMover::FSimulationOutputData& OutputData)
+{
+	Chaos::FCharacterGroundConstraintHandle* ConstraintHandle = nullptr;
+	if (CharacterConstraintProxy && CharacterConstraintProxy->IsInitialized())
+	{
+		ConstraintHandle = CharacterConstraintProxy->GetPhysicsThreadAPI();
+	}
+	if (!ConstraintHandle)
 	{
-		Chaos::EnsureIsInPhysicsThreadContext();
+		return;
+	}
 
-		OnModifyContacts(Modifier);
+	if (CharacterConstraintMode.ShouldEnableConstraint() && !ConstraintHandle->IsEnabled())
+	{
+		EnableCharacterConstraint();
+	}
+	else if (!CharacterConstraintMode.ShouldEnableConstraint() && ConstraintHandle->IsEnabled())
+	{
+		DisableCharacterConstraint();
+		return;
 	}
 
-	void FSimulation::OnPreSimulationTick(const FMoverTimeStep& TimeStep, const FSimulationInputData& InputData)
+	// Update the constraint settings if the mode changes
+	// TODO - Add an initialize for this so we don't have to set it every frame
+	//if (CurrentSyncState.MovementMode != OutputData.SyncState.MovementMode)
 	{
-		if (bInputCmdOverridden)
+		Chaos::FCharacterGroundConstraintSettings& Settings = ConstraintHandle->GetSettings_Mutable();
+		CharacterConstraintMode.UpdateConstraintSettings(Settings);
+	}
+
+	FMoverDefaultSyncState* DefaultSyncState = OutputData.SyncState.SyncStateCollection.FindMutableDataByType<FMoverDefaultSyncState>();
+	check(DefaultSyncState);
+	check(Blackboard.Get());
+
+	FFloorCheckResult FloorResult;
+
+	// Update the constraint data based on the floor result
+	if (Blackboard->TryGet(CommonBlackboard::LastFloorResult, FloorResult) && FloorResult.bBlockingHit)
+	{
+		// Set the ground particle on the constraint
+		Chaos::FGeometryParticleHandle* GroundParticle = nullptr;
+
+		if (Chaos::FPhysicsObjectHandle GroundPhysicsObject = FloorResult.HitResult.PhysicsObject)
 		{
-			InputData.InputCmd = InputCmd;
+			Chaos::FReadPhysicsObjectInterface_Internal ReadInterface = Chaos::FPhysicsObjectInternalInterface::GetRead();
+			if (!ReadInterface.AreAllDisabled({ GroundPhysicsObject }))
+			{
+				GroundParticle = ReadInterface.GetParticle(GroundPhysicsObject);
+				if (ReadInterface.AreAllSleeping({ GroundPhysicsObject }))
+				{
+					Chaos::FWritePhysicsObjectInterface_Internal WriteInterface = Chaos::FPhysicsObjectInternalInterface::GetWrite();
+					WriteInterface.WakeUp({ GroundPhysicsObject });
+				}
+			}
 		}
+		ConstraintHandle->SetGroundParticle(GroundParticle);
+
+		// Set the max walkable slope angle using any override from the hit component
+		float WalkableSlopeCosine = ConstraintHandle->GetSettings().CosMaxWalkableSlopeAngle;
+		if (const TStrongObjectPtr<UPrimitiveComponent> PrimComp = FloorResult.HitResult.Component.Pin())
+		{
+			const FWalkableSlopeOverride& SlopeOverride = PrimComp->GetWalkableSlopeOverride();
+			WalkableSlopeCosine = SlopeOverride.ModifyWalkableFloorZ(WalkableSlopeCosine);
+		}
+
+		if (!FloorResult.bWalkableFloor)
+		{
+			WalkableSlopeCosine = 2.0f;
+		}
+
+		ConstraintHandle->SetData({
+			FloorResult.HitResult.ImpactNormal,
+			DeltaPosition,
+			DeltaRotation.Yaw,
+			FloorResult.FloorDist,
+			WalkableSlopeCosine
+			});
+
+		// Note: We want to record the movement base but we don't record the transform
+		// so don't use this to get a relative transform for the sync state
+		DefaultSyncState->SetMovementBase(FloorResult.HitResult.GetComponent());
 	}
+	else
+	{
+		ConstraintHandle->SetGroundParticle(nullptr);
+		DefaultSyncState->SetMovementBase(nullptr);
+		ConstraintHandle->SetData({
+			ConstraintHandle->GetSettings().VerticalAxis,
+			Chaos::FVec3::ZeroVector,
+			0.0,
+			1.0e10,
+			0.5f
+			});
+	}
+}
 
-	void FSimulation::OnSimulationTick(const FMoverTimeStep& TimeStep, const FSimulationInputData& InputData, FSimulationOutputData& OutputData)
+void UChaosMoverSimulation::OnPostSimulationTick(const FMoverTimeStep& TimeStep, UE::ChaosMover::FSimulationOutputData& OutputData)
+{
+	// TODO - make this more extensible
+	if (TStrongObjectPtr<UBaseMovementMode> CurrentModePtr = StateMachine.GetCurrentMode().Pin())
 	{
-		TStrongObjectPtr<UMoverBlackboard> SimBlackboard = SimBlackboardWeakPtr.Pin();
-		TStrongObjectPtr<UMovementMixer> MovementMixer = MovementMixerWeakPtr.Pin();
+		if (IChaosCharacterMovementModeInterface* CharacterMode = Cast<IChaosCharacterMovementModeInterface>(CurrentModePtr.Get()))
+		{
+			PostSimulationTickCharacter(*CharacterMode, TimeStep, OutputData);
+		}
 
-		if (!SimBlackboard || !MovementMixer)
+		if (IChaosCharacterConstraintMovementModeInterface* CharacterConstraintMode = Cast<IChaosCharacterConstraintMovementModeInterface>(CurrentModePtr.Get()))
+		{
+			PostSimulationTickCharacterConstraint(*CharacterConstraintMode, TimeStep, OutputData);
+		}
+		else
 		{
-			// Looks like the movement component got deleted, no point in running the simulation
-			return;
+			DisableCharacterConstraint();
 		}
+	}
 
-		// JK: I'm not sure this is the right thing to do, but we need a way to pass extra input data that exists in FSimulationInputData that is not InputCmd
-		// but still may be used by the simulation. I chose to use AuxState to pass that data for now, but I may rearrange things so we have a better pipe to send it through
-		// One the things we want to fix is that we shouldn't copy all that info around and write it once exactly where it is read from
-		FMoverAuxStateContext AuxState;
-		AuxState.AuxStateCollection = InputData.InputParams;
-		FMoverTickStartData TickStartData(InputData.InputCmd, SyncState, AuxState);
+	CurrentSyncState = OutputData.SyncState;
+}
 
-		FMoverTickEndData TickEndData;
-		TickEndData.SyncState = SyncState;
-		MoverStateMachine.OnSimulationTick(TimeStep, TickStartData, SimBlackboard.Get(), MovementMixer.Get(), TickEndData);
+void UChaosMoverSimulation::OnModifyContacts(Chaos::FCollisionContactModifier& Modifier) const
+{
+}
 
-		OutputData.SyncState = TickEndData.SyncState;
-		OutputData.LastUsedInputCmd = InputData.InputCmd;
-	}
+const TArray<UE::ChaosMover::FMovementModeChangeEvent>& UChaosMoverSimulation::GetQueuedModeChangeEvents() const
+{
+	return StateMachine.GetQueuedModeChangeEvents();
+}
 
-	void FSimulation::OnPostSimulationTick(const FMoverTimeStep& TimeStep, const FSimulationOutputData& OutputData)
+void UChaosMoverSimulation::ClearQueuedModeChangeEvents()
+{
+	return StateMachine.ClearQueuedModeChangeEvents();
+}
+
+Chaos::FPBDRigidParticleHandle* UChaosMoverSimulation::GetControlledParticle() const
+{
+	if (PhysicsObject)
 	{
+		Chaos::FReadPhysicsObjectInterface_Internal Interface = Chaos::FPhysicsObjectInternalInterface::GetRead();
+		return Interface.GetRigidParticle(PhysicsObject);
 	}
+	return nullptr;
+}
 
-	void FSimulation::OnModifyContacts(Chaos::FCollisionContactModifier& Modifier)
+void UChaosMoverSimulation::EnableCharacterConstraint()
+{
+	Chaos::FCharacterGroundConstraintHandle* ConstraintHandle;
+	if (CharacterConstraintProxy && CharacterConstraintProxy->IsInitialized())
 	{
+		ConstraintHandle = CharacterConstraintProxy->GetPhysicsThreadAPI();
+
+		if (ConstraintHandle->GetCharacterParticle())
+		{
+			ConstraintHandle->SetEnabled(true);
+		}
 	}
+}
 
-} // namespace UE::ChaosMover
\ No newline at end of file
+void UChaosMoverSimulation::DisableCharacterConstraint()
+{
+	Chaos::FCharacterGroundConstraintHandle* ConstraintHandle;
+	if (CharacterConstraintProxy && CharacterConstraintProxy->IsInitialized())
+	{
+		ConstraintHandle = CharacterConstraintProxy->GetPhysicsThreadAPI();
+
+		ConstraintHandle->SetEnabled(false);
+	}
+}
\ No newline at end of file
diff --git a/Engine/Plugins/Experimental/ChaosMover/Source/ChaosMover/Private/ChaosMoverSimulationTypes.cpp b/Engine/Plugins/Experimental/ChaosMover/Source/ChaosMover/Private/ChaosMoverSimulationTypes.cpp
--- a/Engine/Plugins/Experimental/ChaosMover/Source/ChaosMover/Private/ChaosMoverSimulationTypes.cpp
+++ b/Engine/Plugins/Experimental/ChaosMover/Source/ChaosMover/Private/ChaosMoverSimulationTypes.cpp
@@ -8,6 +8,7 @@ void FChaosMoverSimulationDefaultInputs::Reset()
 	CollisionQueryParams = FCollisionQueryParams();
 	UpDir = FVector::UpVector;
 	Gravity = -980.7 * UpDir;
+	PhysicsObjectGravity = 0.0f;
 	PhysicsObject = nullptr;
 	OwningActor = nullptr;
 	CollisionChannel = ECC_Pawn;
diff --git a/Engine/Plugins/Experimental/ChaosMover/Source/ChaosMover/Private/ChaosMoverStateMachine.cpp b/Engine/Plugins/Experimental/ChaosMover/Source/ChaosMover/Private/ChaosMoverStateMachine.cpp
--- a/Engine/Plugins/Experimental/ChaosMover/Source/ChaosMover/Private/ChaosMoverStateMachine.cpp
+++ b/Engine/Plugins/Experimental/ChaosMover/Source/ChaosMover/Private/ChaosMoverStateMachine.cpp
@@ -2,24 +2,20 @@
 
 
 #include "ChaosMoverStateMachine.h"
+
 #include "Backends/ChaosMoverSubsystem.h"
+#include "ChaosMover/ChaosMoverDeveloperSettings.h"
+#include "ChaosMover/ChaosMoverLog.h"
+#include "MoveLibrary/MovementMixer.h"
+#include "MoveLibrary/MovementUtils.h"
 #include "MovementModeTransition.h"
+#include "MoverComponent.h"
 #include "MoverDeveloperSettings.h"
 #include "MoverLog.h"
-#include "MoveLibrary/MovementUtils.h"
-#include "MoverComponent.h"
-#include "MoveLibrary/MovementMixer.h"
 #include "Templates/SubclassOf.h"
 
 namespace UE::ChaosMover
 {
-	static int32 SkipGenerateMoveIfOverridden = 1;
-	FAutoConsoleVariableRef CVarSkipGenerateMoveIfOverridden(
-		TEXT("chaosmover.perf.SkipGenerateMoveIfOverridden"),
-		SkipGenerateMoveIfOverridden,
-		TEXT("If true and we have a layered move fully overriding movement, then we will skip calling OnGenerateMove on the active movement mode for better performance\n")
-	);
-
 	FMoverStateMachine::FMoverStateMachine()
 	{
 	}
@@ -28,12 +24,12 @@ namespace UE::ChaosMover
 	{
 	}
 
-	void FMoverStateMachine::RegisterMovementMode(FName ModeName, TObjectPtr<UBaseMovementMode> Mode, bool bIsDefaultMode)
+	void FMoverStateMachine::RegisterMovementMode(FName ModeName, TWeakObjectPtr<UBaseMovementMode> Mode, bool bIsDefaultMode)
 	{
 		// JAH TODO: add validation and warnings for overwriting modes
 		// JAH TODO: add validation of Mode
 
-		Modes.Add(ModeName, TWeakObjectPtr<UBaseMovementMode>(Mode.Get()));
+		Modes.Add(ModeName, Mode);
 
 		if (bIsDefaultMode)
 		{
@@ -79,18 +75,18 @@ namespace UE::ChaosMover
 		DefaultModeName = NewDefaultModeName;
 	}
 
-	void FMoverStateMachine::RegisterGlobalTransition(TObjectPtr<UBaseMovementModeTransition> Transition)
+	void FMoverStateMachine::RegisterGlobalTransition(TWeakObjectPtr<UBaseMovementModeTransition> Transition)
 	{
-		GlobalTransitions.Add(TWeakObjectPtr<UBaseMovementModeTransition>(Transition.Get()));
+		GlobalTransitions.Add(Transition);
 
 		Transition->OnRegistered();
 	}
 
-	void FMoverStateMachine::UnregisterGlobalTransition(TObjectPtr<UBaseMovementModeTransition> Transition)
+	void FMoverStateMachine::UnregisterGlobalTransition(TWeakObjectPtr<UBaseMovementModeTransition> Transition)
 	{
 		Transition->OnUnregistered();
 
-		GlobalTransitions.Remove(TWeakObjectPtr<UBaseMovementModeTransition>(Transition.Get()));
+		GlobalTransitions.Remove(Transition);
 	}
 
 	void FMoverStateMachine::ClearAllGlobalTransitions()
@@ -120,7 +116,7 @@ namespace UE::ChaosMover
 				if ((NextModeName != NAME_None) &&
 					(NextModeName != DesiredNextModeName || bShouldReenter != bShouldNextModeReenter))
 				{
-					UE_LOG(LogMover, Log, TEXT("%s (%s) Overwriting of queued mode change (%s, reenter: %i) with (%s, reenter: %i)"), *OwnerActorName, *UEnum::GetValueAsString(OwnerActorLocalNetRole), *NextModeName.ToString(), bShouldNextModeReenter, *DesiredNextModeName.ToString(), bShouldReenter);
+					UE_LOG(LogChaosMover, Log, TEXT("%s (%s) Overwriting of queued mode change (%s, reenter: %i) with (%s, reenter: %i)"), *OwnerActorName, *UEnum::GetValueAsString(OwnerActorLocalNetRole), *NextModeName.ToString(), bShouldNextModeReenter, *DesiredNextModeName.ToString(), bShouldReenter);
 				}
 
 				if (Modes.Contains(DesiredNextModeName))
@@ -129,7 +125,7 @@ namespace UE::ChaosMover
 				}
 				else
 				{
-					UE_LOG(LogMover, Warning, TEXT("Attempted to queue an unregistered movement mode: %s on owner %s"), *DesiredNextModeName.ToString(), *OwnerActorName);
+					UE_LOG(LogChaosMover, Warning, TEXT("Attempted to queue an unregistered movement mode: %s on owner %s"), *DesiredNextModeName.ToString(), *OwnerActorName);
 				}
 			}
 		}
@@ -254,11 +250,11 @@ namespace UE::ChaosMover
 			// Merge proposed movement from the current mode with movement from layered moves
 			if (!CurrentModeName.IsNone() && Modes.Contains(CurrentModeName))
 			{
-				UBaseMovementMode* CurrentMode = Modes[CurrentModeName].Get();
+				TStrongObjectPtr<UBaseMovementMode> CurrentMode = Modes[CurrentModeName].Pin();
 				FProposedMove CombinedMove;
 				bool bHasModeMoveContribution = false;
 
-				if (!SkipGenerateMoveIfOverridden ||
+				if (!CVars::bSkipGenerateMoveIfOverridden ||
 					!(bHasLayeredMoveContributions && CombinedLayeredMove.MixMode == EMoveMixMode::OverrideAll))
 				{
 					QUICK_SCOPE_CYCLE_COUNTER(STAT_GenerateMoveFromMode);
@@ -364,7 +360,7 @@ namespace UE::ChaosMover
 					// if we've done this sub step a lot before go ahead and just advance time to avoid freezing editor
 					if (NumConsecutiveFullRefundedSubsteps >= MaxConsecutiveFullRefundedSubsteps)
 					{
-						UE_LOG(LogMover, Warning, TEXT("Movement mode %s and %s on %s are stuck giving time back to each other. Overriding to advance to next substep."),
+						UE_LOG(LogChaosMover, Warning, TEXT("Movement mode %s and %s on %s are stuck giving time back to each other. Overriding to advance to next substep."),
 							*CurrentModeName.ToString(),
 							*OutputState.MovementEndState.NextModeName.ToString(),
 							*OwnerActorName);
@@ -417,21 +413,21 @@ namespace UE::ChaosMover
 	}
 
 
-	const UBaseMovementMode* FMoverStateMachine::GetCurrentMode() const
+	const TWeakObjectPtr<UBaseMovementMode> FMoverStateMachine::GetCurrentMode() const
 	{
 		if (CurrentModeName != NAME_None && Modes.Contains(CurrentModeName))
 		{
-			return Modes[CurrentModeName].Get();
+			return Modes[CurrentModeName];
 		}
 
 		return nullptr;
 	}
 
-	const UBaseMovementMode* FMoverStateMachine::FindMovementMode(FName ModeName) const
+	const TWeakObjectPtr<UBaseMovementMode> FMoverStateMachine::FindMovementMode(FName ModeName) const
 	{
 		if (ModeName != NAME_None && Modes.Contains(ModeName))
 		{
-			return Modes[ModeName].Get();
+			return Modes[ModeName];
 		}
 
 		return nullptr;
@@ -458,10 +454,10 @@ namespace UE::ChaosMover
 		ClearQueuedMode();
 	}
 
-	void FMoverStateMachine::InitDefaultModesAndTransitions(TObjectPtr<UImmediateMovementModeTransition> InImmediateMovementModeTransition, TObjectPtr<UNullMovementMode> InNullMovementMode)
+	void FMoverStateMachine::InitDefaultModesAndTransitions(const TWeakObjectPtr<UImmediateMovementModeTransition> InImmediateMovementModeTransition, const TWeakObjectPtr<UNullMovementMode> InNullMovementMode)
 	{
-		ImmediateMovementModeTransitionWeakPtr = TWeakObjectPtr<UImmediateMovementModeTransition>(InImmediateMovementModeTransition.Get());
-		NullMovementModeWeakPtr = TWeakObjectPtr<UNullMovementMode>(InNullMovementMode.Get());
+		ImmediateMovementModeTransitionWeakPtr = InImmediateMovementModeTransition;
+		NullMovementModeWeakPtr = InNullMovementMode;
 	}
 
 	void FMoverStateMachine::AdvanceToNextMode()
@@ -484,7 +480,7 @@ namespace UE::ChaosMover
 
 						CurrentModeName = NextModeName;
 
-						UE_LOG(LogMover, Verbose, TEXT("AdvanceToNextMode: %s (%s) from %s to %s"),
+						UE_LOG(LogChaosMover, Verbose, TEXT("AdvanceToNextMode: %s (%s) from %s to %s"),
 							*OwnerActorName, *UEnum::GetValueAsString(OwnerActorLocalNetRole), *CurrentModeName.ToString(), *NextModeName.ToString());
 
 						TWeakObjectPtr<UBaseMovementMode> PreviousMovementMode = FoundPreviousMovementMode ? *FoundPreviousMovementMode : TWeakObjectPtr<UBaseMovementMode>();
diff --git a/Engine/Plugins/Experimental/ChaosMover/Source/ChaosMover/Private/ChaosMoverStateMachine.h b/Engine/Plugins/Experimental/ChaosMover/Source/ChaosMover/Private/ChaosMoverStateMachine.h
--- a/Engine/Plugins/Experimental/ChaosMover/Source/ChaosMover/Private/ChaosMoverStateMachine.h
+++ b/Engine/Plugins/Experimental/ChaosMover/Source/ChaosMover/Private/ChaosMoverStateMachine.h
@@ -4,9 +4,9 @@
 
 #include "CoreMinimal.h"
 #include "Engine/EngineTypes.h"
-#include "Templates/SubclassOf.h"
-#include "MovementMode.h"
 #include "InstantMovementEffect.h"
+#include "MovementMode.h"
+#include "Templates/SubclassOf.h"
 
 namespace UE::ChaosMover
 {
@@ -25,13 +25,13 @@ namespace UE::ChaosMover
 		FMoverStateMachine();
 		virtual ~FMoverStateMachine();
 
-		void RegisterMovementMode(FName ModeName, TObjectPtr<UBaseMovementMode> Mode, bool bIsDefaultMode = false);
+		void RegisterMovementMode(FName ModeName, TWeakObjectPtr<UBaseMovementMode> Mode, bool bIsDefaultMode = false);
 
 		void UnregisterMovementMode(FName ModeName);
 		void ClearAllMovementModes();
 
-		void RegisterGlobalTransition(TObjectPtr<UBaseMovementModeTransition> Transition);
-		void UnregisterGlobalTransition(TObjectPtr<UBaseMovementModeTransition> Transition);
+		void RegisterGlobalTransition(TWeakObjectPtr<UBaseMovementModeTransition> Transition);
+		void UnregisterGlobalTransition(TWeakObjectPtr<UBaseMovementModeTransition> Transition);
 		void ClearAllGlobalTransitions();
 
 		void SetDefaultMode(FName NewDefaultModeName);
@@ -44,11 +44,14 @@ namespace UE::ChaosMover
 		void OnSimulationPreRollback(const FMoverSyncState* InvalidSyncState, const FMoverSyncState* SyncState, const FMoverAuxStateContext* InvalidAuxState, const FMoverAuxStateContext* AuxState);
 		void OnSimulationRollback(const FMoverSyncState* SyncState, const FMoverAuxStateContext* AuxState);
 
-		FName GetCurrentModeName() const { return CurrentModeName; }
+		FName GetCurrentModeName() const
+		{
+			return CurrentModeName;
+		}
 
-		const UBaseMovementMode* GetCurrentMode() const;
+		const TWeakObjectPtr<UBaseMovementMode> GetCurrentMode() const;
 
-		const UBaseMovementMode* FindMovementMode(FName ModeName) const;
+		const TWeakObjectPtr<UBaseMovementMode> FindMovementMode(FName ModeName) const;
 
 		void QueueLayeredMove(TSharedPtr<FLayeredMoveBase> Move);
 
@@ -60,7 +63,7 @@ namespace UE::ChaosMover
 		void SetOwnerActorName(const FString& InOwnerActorName);
 		void SetOwnerActorLocalNetRole(ENetRole InOwnerActorLocalNetRole);
 
-		void InitDefaultModesAndTransitions(TObjectPtr<UImmediateMovementModeTransition> ImmediateMovementModeTransition, TObjectPtr<UNullMovementMode> NullMovementMode);
+		void InitDefaultModesAndTransitions(const TWeakObjectPtr<UImmediateMovementModeTransition> ImmediateMovementModeTransition, const TWeakObjectPtr<UNullMovementMode> NullMovementMode);
 
 	protected:
 
diff --git a/Engine/Plugins/Experimental/ChaosMover/Source/ChaosMover/Private/Character/Modes/ChaosCharacterMovementMode.cpp b/Engine/Plugins/Experimental/ChaosMover/Source/ChaosMover/Private/Character/Modes/ChaosCharacterMovementMode.cpp
--- a/Engine/Plugins/Experimental/ChaosMover/Source/ChaosMover/Private/Character/Modes/ChaosCharacterMovementMode.cpp
+++ b/Engine/Plugins/Experimental/ChaosMover/Source/ChaosMover/Private/Character/Modes/ChaosCharacterMovementMode.cpp
@@ -90,3 +90,13 @@ void UChaosCharacterMovementMode::UpdateConstraintSettings(Chaos::FCharacterGrou
 	ConstraintSettings.SwingTorqueLimit = FUnitConversion::Convert(SwingTorqueLimit, EUnit::NewtonMeters, EUnit::KilogramCentimetersSquaredPerSecondSquared);
 	ConstraintSettings.TargetHeight = TargetHeight;
 }
+
+float UChaosCharacterMovementMode::GetMaxWalkSlopeCosine() const
+{
+	if (const USharedChaosCharacterMovementSettings* SharedSettingsPtr = SharedSettings.Get())
+	{
+		return SharedSettingsPtr->GetMaxWalkableSlopeCosine();
+	}
+
+	return 0.707f;
+}
diff --git a/Engine/Plugins/Experimental/ChaosMover/Source/ChaosMover/Private/Character/Modes/ChaosFallingMode.cpp b/Engine/Plugins/Experimental/ChaosMover/Source/ChaosMover/Private/Character/Modes/ChaosFallingMode.cpp
--- a/Engine/Plugins/Experimental/ChaosMover/Source/ChaosMover/Private/Character/Modes/ChaosFallingMode.cpp
+++ b/Engine/Plugins/Experimental/ChaosMover/Source/ChaosMover/Private/Character/Modes/ChaosFallingMode.cpp
@@ -2,23 +2,174 @@
 
 #include "ChaosMover/Character/Modes/ChaosFallingMode.h"
 
+#include "ChaosMover/ChaosMoverLog.h"
+#include "ChaosMover/ChaosMoverSimulation.h"
+#include "ChaosMover/Character/Settings/SharedChaosCharacterMovementSettings.h"
+#include "MoveLibrary/AirMovementUtils.h"
+#include "MoveLibrary/FloorQueryUtils.h"
+#include "MoveLibrary/MovementUtils.h"
+
 #include UE_INLINE_GENERATED_CPP_BY_NAME(ChaosFallingMode)
 
 UChaosFallingMode::UChaosFallingMode(const FObjectInitializer& ObjectInitializer)
 	: Super(ObjectInitializer)
+	, bCancelVerticalSpeedOnLanding(true)
+	, AirControlPercentage(0.4f)
+	, FallingDeceleration(200.0f)
+	, OverTerminalSpeedFallingDeceleration(800.0f)
+	, TerminalMovementPlaneSpeed(1500.0f)
+	, bShouldClampTerminalVerticalSpeed(true)
+	, VerticalFallingDeceleration(4000.0f)
+	, TerminalVerticalSpeed(2000.0f)
 {
 	bSupportsAsync = true;
 
 	GameplayTags.AddTag(Mover_IsInAir);
 	GameplayTags.AddTag(Mover_IsFalling);
+	GameplayTags.AddTag(Mover_SkipAnimRootMotion);
 }
 
 void UChaosFallingMode::GenerateMove_Implementation(const FMoverTickStartData& StartState, const FMoverTimeStep& TimeStep, FProposedMove& OutProposedMove) const
 {
+	const FChaosMoverSimulationDefaultInputs* DefaultSimInputs = StartState.AuxState.AuxStateCollection.FindDataByType<FChaosMoverSimulationDefaultInputs>();
+	const FCharacterDefaultInputs* CharacterInputs = StartState.InputCmd.InputCollection.FindDataByType<FCharacterDefaultInputs>();
+	if (!DefaultSimInputs || !CharacterInputs)
+	{
+		UE_LOG(LogChaosMover, Warning, TEXT("ChaosFallingMode requires FChaosMoverSimulationDefaultInputs and FCharacterDefaultInputs"));
+		return;
+	}
+
+	if (!Simulation)
+	{
+		UE_LOG(LogChaosMover, Warning, TEXT("No Simulation set on ChaosFallingMode"));
+		return;
+	}
+
+	const FMoverDefaultSyncState* StartingSyncState = StartState.SyncState.SyncStateCollection.FindDataByType<FMoverDefaultSyncState>();
+	check(StartingSyncState);
+
+	const float DeltaSeconds = TimeStep.StepMs * 0.001f;
+
+	// We don't want velocity limits to take the falling velocity component into account, since it is handled 
+	// separately by the terminal velocity of the environment.
+	const FVector StartVelocity = StartingSyncState->GetVelocity_WorldSpace();
+	const FVector StartHorizontalVelocity = FVector::VectorPlaneProject(StartVelocity, DefaultSimInputs->UpDir);
+
+	FFreeMoveParams Params;
+	Params.MoveInputType = CharacterInputs->GetMoveInputType();
+	const bool bMaintainInputMagnitude = true;
+	Params.MoveInput = UPlanarConstraintUtils::ConstrainDirectionToPlane(FPlanarConstraint(), CharacterInputs->GetMoveInput_WorldSpace(), bMaintainInputMagnitude);
+	Params.MoveInput *= AirControlPercentage;
+	// Don't care about up axis input since falling - if up input matters that should probably be a different movement mode
+	Params.MoveInput = FVector::VectorPlaneProject(Params.MoveInput, DefaultSimInputs->UpDir);
+
+	FRotator IntendedOrientation_WorldSpace;
+	// If there's no intent from input to change orientation, use the current orientation
+	if (!CharacterInputs || CharacterInputs->OrientationIntent.IsNearlyZero())
+	{
+		IntendedOrientation_WorldSpace = StartingSyncState->GetOrientation_WorldSpace();
+	}
+	else
+	{
+		IntendedOrientation_WorldSpace = CharacterInputs->GetOrientationIntentDir_WorldSpace().ToOrientationRotator();
+	}
+
+	FQuat WorldToGravityTransform = FQuat::FindBetweenNormals(FVector::UpVector, DefaultSimInputs->UpDir);
+	IntendedOrientation_WorldSpace = UMovementUtils::ApplyGravityToOrientationIntent(IntendedOrientation_WorldSpace, WorldToGravityTransform, bShouldCharacterRemainUpright);
+
+	Params.OrientationIntent = IntendedOrientation_WorldSpace;
+	Params.PriorVelocity = StartHorizontalVelocity;
+	Params.PriorOrientation = StartingSyncState->GetOrientation_WorldSpace();
+	Params.DeltaSeconds = DeltaSeconds;
+	Params.TurningRate = SharedSettings->TurningRate;
+	Params.TurningBoost = SharedSettings->TurningBoost;
+	Params.MaxSpeed = SharedSettings->MaxSpeed;
+	Params.Acceleration = SharedSettings->Acceleration;
+	Params.Deceleration = FallingDeceleration;
+	Params.WorldToGravityQuat = WorldToGravityTransform;
+	Params.bUseAccelerationForVelocityMove = SharedSettings->bUseAccelerationForVelocityMove;
 
+	// Check if any current velocity values are over our terminal velocity - if so limit the move input in that direction and apply OverTerminalVelocityFallingDeceleration
+	if (Params.MoveInput.Dot(StartVelocity) > 0 && StartHorizontalVelocity.Size() >= TerminalMovementPlaneSpeed)
+	{
+		Params.Deceleration = OverTerminalSpeedFallingDeceleration;
+	}
+
+	FFloorCheckResult LastFloorResult;
+	// limit our move input based on the floor we're on
+	const UMoverBlackboard* SimBlackboard = Simulation->GetBlackboard();
+	if (SimBlackboard && SimBlackboard->TryGet(CommonBlackboard::LastFloorResult, LastFloorResult))
+	{
+		if (LastFloorResult.HitResult.IsValidBlockingHit() && LastFloorResult.HitResult.Normal.Dot(DefaultSimInputs->UpDir) > UE::MoverUtils::VERTICAL_SLOPE_NORMAL_MAX_DOT && !LastFloorResult.IsWalkableFloor())
+		{
+			// If acceleration is into the wall, limit contribution.
+			if (FVector::DotProduct(Params.MoveInput, LastFloorResult.HitResult.Normal) < 0.f)
+			{
+				// Allow movement parallel to the wall, but not into it because that may push us up.
+				const FVector FallingHitNormal = FVector::VectorPlaneProject(LastFloorResult.HitResult.Normal, -DefaultSimInputs->UpDir).GetSafeNormal();
+				Params.MoveInput = FVector::VectorPlaneProject(Params.MoveInput, FallingHitNormal);
+			}
+		}
+	}
+
+	OutProposedMove = UAirMovementUtils::ComputeControlledFreeMove(Params);
+	const FVector VelocityWithGravity = StartVelocity + UMovementUtils::ComputeVelocityFromGravity(DefaultSimInputs->Gravity, DeltaSeconds);
+
+	//  If we are going faster than TerminalVerticalVelocity apply VerticalFallingDeceleration otherwise reset Z velocity to before we applied deceleration 
+	if (VelocityWithGravity.GetAbs().Dot(DefaultSimInputs->UpDir) > TerminalVerticalSpeed)
+	{
+		if (bShouldClampTerminalVerticalSpeed)
+		{
+			const float ClampedVerticalSpeed = FMath::Sign(VelocityWithGravity.Dot(DefaultSimInputs->UpDir)) * TerminalVerticalSpeed;
+			UMovementUtils::SetGravityVerticalComponent(OutProposedMove.LinearVelocity, ClampedVerticalSpeed, DefaultSimInputs->UpDir);
+		}
+		else
+		{
+			float DesiredDeceleration = FMath::Abs(TerminalVerticalSpeed - VelocityWithGravity.GetAbs().Dot(DefaultSimInputs->UpDir)) / DeltaSeconds;
+			float DecelerationToApply = FMath::Min(DesiredDeceleration, VerticalFallingDeceleration);
+			DecelerationToApply = FMath::Sign(VelocityWithGravity.Dot(DefaultSimInputs->UpDir)) * DecelerationToApply * DeltaSeconds;
+			FVector MaxUpDirVelocity = VelocityWithGravity * DefaultSimInputs->UpDir - (DefaultSimInputs->UpDir * DecelerationToApply);
+
+			UMovementUtils::SetGravityVerticalComponent(OutProposedMove.LinearVelocity, MaxUpDirVelocity.Dot(DefaultSimInputs->UpDir), DefaultSimInputs->UpDir);
+		}
+	}
+	else
+	{
+		UMovementUtils::SetGravityVerticalComponent(OutProposedMove.LinearVelocity, VelocityWithGravity.Dot(DefaultSimInputs->UpDir), DefaultSimInputs->UpDir);
+	}
 }
 
 void UChaosFallingMode::SimulationTick_Implementation(const FSimulationTickParams& Params, FMoverTickEndData& OutputState)
 {
+	const FChaosMoverSimulationDefaultInputs* DefaultSimInputs = Params.StartState.AuxState.AuxStateCollection.FindDataByType<FChaosMoverSimulationDefaultInputs>();
+	if (!DefaultSimInputs)
+	{
+		UE_LOG(LogChaosMover, Warning, TEXT("ChaosFallingMode requires FChaosMoverSimulationDefaultInputs"));
+		return;
+	}
+
+	const FMoverDefaultSyncState* StartingSyncState = Params.StartState.SyncState.SyncStateCollection.FindDataByType<FMoverDefaultSyncState>();
+	check(StartingSyncState);
+
+	FProposedMove ProposedMove = Params.ProposedMove;
+	FMoverDefaultSyncState& OutputSyncState = OutputState.SyncState.SyncStateCollection.FindOrAddMutableDataByType<FMoverDefaultSyncState>();
+
+	const float DeltaSeconds = Params.TimeStep.StepMs * 0.001f;
+
+	FRotator TargetOrient = StartingSyncState->GetOrientation_WorldSpace();
+	if (!UMovementUtils::IsAngularVelocityZero(ProposedMove.AngularVelocity))
+	{
+		TargetOrient += (ProposedMove.AngularVelocity * DeltaSeconds);
+	}
+
+	// The physics simulation applies Z-only gravity acceleration via physics volumes, so we need to account for it here 
+	FVector TargetVel = ProposedMove.LinearVelocity - DefaultSimInputs->PhysicsObjectGravity * DefaultSimInputs->UpDir * DeltaSeconds;
+	FVector TargetPos = StartingSyncState->GetLocation_WorldSpace() + TargetVel * DeltaSeconds;
 
+	OutputState.MovementEndState.RemainingMs = 0.0f;
+	OutputSyncState.MoveDirectionIntent = ProposedMove.bHasDirIntent ? ProposedMove.DirectionIntent : FVector::ZeroVector;
+	OutputSyncState.SetTransforms_WorldSpace(
+		TargetPos,
+		TargetOrient,
+		TargetVel);
 }
\ No newline at end of file
diff --git a/Engine/Plugins/Experimental/ChaosMover/Source/ChaosMover/Private/Character/Transitions/ChaosCharacterLandingCheck.cpp b/Engine/Plugins/Experimental/ChaosMover/Source/ChaosMover/Private/Character/Transitions/ChaosCharacterLandingCheck.cpp
new file mode 100644
--- /dev/null
+++ b/Engine/Plugins/Experimental/ChaosMover/Source/ChaosMover/Private/Character/Transitions/ChaosCharacterLandingCheck.cpp
@@ -0,0 +1,112 @@
+// Copyright Epic Games, Inc. All Rights Reserved.
+
+#include "ChaosMover/Character/Transitions/ChaosCharacterLandingCheck.h"
+
+#include "ChaosMover/ChaosMoverLog.h"
+#include "ChaosMover/ChaosMoverSimulation.h"
+#include "ChaosMover/ChaosMoverSimulationTypes.h"
+#include "ChaosMover/Utilities/ChaosGroundMovementUtils.h"
+#include "MoveLibrary/FloorQueryUtils.h"
+#include "MoveLibrary/WaterMovementUtils.h"
+#include "MoverSimulationTypes.h"
+
+#include UE_INLINE_GENERATED_CPP_BY_NAME(ChaosCharacterLandingCheck)
+
+UChaosCharacterLandingCheck::UChaosCharacterLandingCheck(const FObjectInitializer& ObjectInitializer)
+	: Super(ObjectInitializer)
+{
+	bSupportsAsync = true;
+
+	TransitionToGroundMode = DefaultModeNames::Walking;
+	TransitionToWaterMode = NAME_None;
+}
+
+FTransitionEvalResult UChaosCharacterLandingCheck::Evaluate_Implementation(const FSimulationTickParams& Params) const
+{
+	check(Simulation);
+
+	FTransitionEvalResult EvalResult;
+
+	const FMoverTickStartData& StartState = Params.StartState;
+
+	if (const FChaosMoverSimulationDefaultInputs* SimInputs = StartState.AuxState.AuxStateCollection.FindDataByType<FChaosMoverSimulationDefaultInputs>())
+	{
+		const UMoverBlackboard* Blackboard = Simulation->GetBlackboard();
+		check(Blackboard);
+
+		const FCharacterDefaultInputs* CharacterInputs = StartState.InputCmd.InputCollection.FindDataByType<FCharacterDefaultInputs>();
+		const FMoverDefaultSyncState* StartingSyncState = StartState.SyncState.SyncStateCollection.FindDataByType<FMoverDefaultSyncState>();
+		check(CharacterInputs);
+		check(StartingSyncState);
+
+		if (const IChaosCharacterMovementModeInterface* Mode = Cast<IChaosCharacterMovementModeInterface>(Simulation->FindMovementModeByName(Params.StartState.SyncState.MovementMode)))
+		{
+			const float DeltaSeconds = Params.TimeStep.StepMs * 0.001f;
+			const FVector LinearVelocity = StartingSyncState->GetVelocity_WorldSpace();
+			const float VerticalVelocity = SimInputs->UpDir.Dot(LinearVelocity);
+			const bool bJumping = CharacterInputs && CharacterInputs->bIsJumpPressed;
+			const bool bIsMovingUp = bJumping || (VerticalVelocity > 0.0f);
+
+			// Check for water landing
+			if (!TransitionToWaterMode.IsNone())
+			{
+				FWaterCheckResult WaterResult;
+				if (Blackboard->TryGet(CommonBlackboard::LastFloorResult, WaterResult) && WaterResult.IsSwimmableVolume())
+				{
+					const float ProjectedImmersionDepth = WaterResult.WaterSplineData.ImmersionDepth - VerticalVelocity * DeltaSeconds;
+					const bool bStartSwimming = ProjectedImmersionDepth > SwimmingStartImmersionDepth;
+
+					if (bStartSwimming && !bIsMovingUp)
+					{
+						if (Simulation->FindMovementModeByName(TransitionToWaterMode))
+						{
+							EvalResult.NextMode = TransitionToWaterMode;
+						}
+						else
+						{
+							UE_LOG(LogChaosMover, Warning, TEXT("Invalid water mode name %s in ChaosCharacterLandingCheck. Cannot make transition"), *TransitionToWaterMode.ToString());
+						}
+
+						return EvalResult;
+					}
+				}
+			}
+
+			// Check for ground landing
+			if (!TransitionToGroundMode.IsNone())
+			{
+				FFloorCheckResult FloorResult;
+				if (Blackboard->TryGet(CommonBlackboard::LastFloorResult, FloorResult) && FloorResult.IsWalkableFloor())
+				{
+					const float TargetHeight = Mode->GetTargetHeight();
+					const float FloorDistanceWithFloorNormal = FloorResult.HitResult.ImpactNormal.Dot(StartingSyncState->GetLocation_WorldSpace() - FloorResult.HitResult.ImpactPoint);
+					const FVector LocalGroundVelocity = UChaosGroundMovementUtils::ComputeLocalGroundVelocity_Internal(StartingSyncState->GetLocation_WorldSpace(), FloorResult);
+					const float RelativeVerticalVelocity = FloorResult.HitResult.ImpactNormal.Dot(LinearVelocity - LocalGroundVelocity);
+					const float ProjectedFloorDistance = FloorDistanceWithFloorNormal + RelativeVerticalVelocity * DeltaSeconds;
+					const bool bIsFloorWithinReach = ProjectedFloorDistance < TargetHeight + FloorDistanceTolerance;
+					const bool bIsMovingUpRelativeToFloor = (RelativeVerticalVelocity > UE_KINDA_SMALL_NUMBER) || bJumping;
+
+					if (bIsFloorWithinReach && !bIsMovingUpRelativeToFloor)
+					{
+						if (Simulation->FindMovementModeByName(TransitionToGroundMode))
+						{
+							EvalResult.NextMode = TransitionToGroundMode;
+						}
+						else
+						{
+							UE_LOG(LogChaosMover, Warning, TEXT("Invalid ground mode name %s in ChaosCharacterLandingCheck. Cannot make transition"), *TransitionToGroundMode.ToString());
+						}
+						
+						return EvalResult;
+					}
+				}
+			}
+		}
+	}
+
+	return EvalResult;
+}
+
+void UChaosCharacterLandingCheck::Trigger_Implementation(const FSimulationTickParams& Params)
+{
+}
\ No newline at end of file
diff --git a/Engine/Plugins/Experimental/ChaosMover/Source/ChaosMover/Private/Utilities/ChaosGroundMovementUtils.cpp b/Engine/Plugins/Experimental/ChaosMover/Source/ChaosMover/Private/Utilities/ChaosGroundMovementUtils.cpp
new file mode 100644
--- /dev/null
+++ b/Engine/Plugins/Experimental/ChaosMover/Source/ChaosMover/Private/Utilities/ChaosGroundMovementUtils.cpp
@@ -0,0 +1,30 @@
+// Copyright Epic Games, Inc. All Rights Reserved.
+
+#include "ChaosMover/Utilities/ChaosGroundMovementUtils.h"
+
+#include "Chaos/PhysicsObjectInternalInterface.h"
+#include "MoveLibrary/FloorQueryUtils.h"
+
+FVector UChaosGroundMovementUtils::ComputeLocalGroundVelocity_Internal(const FVector& Position, const FFloorCheckResult& FloorResult)
+{
+	FVector GroundVelocity = FVector::ZeroVector;
+	if (const Chaos::FPBDRigidParticleHandle* Rigid = GetRigidParticleHandleFromFloorResult_Internal(FloorResult))
+	{
+		Chaos::FRigidTransform3 ComTransform = Rigid->GetTransformXRCom();
+		FVector Offset = Position - ComTransform.GetLocation();
+		Offset -= Offset.ProjectOnToNormal(FloorResult.HitResult.ImpactNormal);
+		GroundVelocity = Rigid->GetV() + Rigid->GetW().Cross(Offset);
+	}
+	return GroundVelocity;
+}
+
+Chaos::FPBDRigidParticleHandle* UChaosGroundMovementUtils::GetRigidParticleHandleFromFloorResult_Internal(const FFloorCheckResult& FloorResult)
+{
+	if (Chaos::FPhysicsObjectHandle PhysicsObject = FloorResult.HitResult.PhysicsObject)
+	{
+		Chaos::FReadPhysicsObjectInterface_Internal Interface = Chaos::FPhysicsObjectInternalInterface::GetRead();
+		return Interface.GetRigidParticle(PhysicsObject);
+	}
+
+	return nullptr;
+}
diff --git a/Engine/Plugins/Experimental/ChaosMover/Source/ChaosMover/Public/ChaosMover/Backends/ChaosMoverBackend.h b/Engine/Plugins/Experimental/ChaosMover/Source/ChaosMover/Public/ChaosMover/Backends/ChaosMoverBackend.h
--- a/Engine/Plugins/Experimental/ChaosMover/Source/ChaosMover/Public/ChaosMover/Backends/ChaosMoverBackend.h
+++ b/Engine/Plugins/Experimental/ChaosMover/Source/ChaosMover/Public/ChaosMover/Backends/ChaosMoverBackend.h
@@ -3,20 +3,23 @@
 #pragma once
 
 #include "Backends/MoverBackendLiaison.h"
+#include "Chaos/Character/CharacterGroundConstraint.h"
 #include "Chaos/Declares.h"
+#include "ChaosMover/ChaosMoverSimulation.h"
 #include "ChaosMover/ChaosMoverSimulationTypes.h"
 #include "Components/ActorComponent.h"
 
 #include "ChaosMoverBackend.generated.h"
 
-namespace UE::ChaosMover
-{
-	class FSimulation;
-}
-
 class UMoverComponent;
 class UNetworkPhysicsComponent;
 
+namespace Chaos
+{
+	class FCharacterGroundConstraint;
+	class FCharacterGroundConstraintProxy;
+}
+
 UCLASS(Within = MoverComponent)
 class CHAOSMOVER_API UChaosMoverBackendComponent : public UActorComponent, public IMoverBackendLiaisonInterface
 {
@@ -46,17 +49,24 @@ public:
 	virtual void ConsumeOutputData(const FMoverTimeStep& TimeStep, const UE::ChaosMover::FSimulationOutputData& OutputData);
 	virtual void FinalizeFrame(double ResultsTimeInSeconds);
 
-	void ApplyNetInputData(const FMoverInputCmdContext& InNetInputCmd);
-	void BuildNetInputData(FMoverInputCmdContext& OutNetInputCmd) const;
-	void ApplyNetStateData(const FMoverSyncState& InNetSyncState);
-	void BuildNetStateData(FMoverSyncState& OutNetSyncState) const;
-	void RegisterSimulation() const;
+	UChaosMoverSimulation* GetSimulation()
+	{
+		return Simulation;
+	}
+
+	const UChaosMoverSimulation* GetSimulation() const
+	{
+		return Simulation;
+	}
 
 protected:
-	virtual TSharedPtr<UE::ChaosMover::FSimulation> CreateSimulation();
 	virtual void InitSimulation();
-	void InitSimulationStateMachine();
+	virtual void DeinitSimulation();
+	virtual void CreatePhysics();
+	virtual void DestroyPhysics();
+
 	virtual void GenerateInput(const FMoverTimeStep& TimeStep, UE::ChaosMover::FSimulationInputData& InputData);
+
 	/**
 	 * Flushes the queue of movement mode change events, which causes all movement mode change callbacks to be called
 	 */
@@ -76,14 +86,19 @@ protected:
 	UPROPERTY(Transient)
 	TObjectPtr<UNetworkPhysicsComponent> NetworkPhysicsComponent;
 
-	TSharedPtr<UE::ChaosMover::FSimulation> Simulation;
-	UE::Mover::FSimulationOutputRecord SimOutputRecord;
-	bool bIsUsingAsyncPhysics = false;
-	bool bIsInputProducer = true;
+	UPROPERTY(Transient)
+	TObjectPtr<UNullMovementMode> NullMovementMode;
 
 	UPROPERTY(Transient)
-	TObjectPtr<UMoverBlackboard> MoverBlackboard_Async;
+	TObjectPtr<UImmediateMovementModeTransition> ImmediateModeTransition;
 
 	UPROPERTY(Transient)
-	TObjectPtr<UMovementMixer> MovementMixer_Async;
+	TObjectPtr<UChaosMoverSimulation> Simulation;
+
+	UE::Mover::FSimulationOutputRecord SimOutputRecord;
+	TUniquePtr<Chaos::FCharacterGroundConstraint> CharacterGroundConstraint;
+
+	bool bIsUsingAsyncPhysics = false;
+	bool bWantsDestroySim = false;
+	bool bWantsCreateSim = false;
 };
\ No newline at end of file
diff --git a/Engine/Plugins/Experimental/ChaosMover/Source/ChaosMover/Public/ChaosMover/ChaosMovementMode.h b/Engine/Plugins/Experimental/ChaosMover/Source/ChaosMover/Public/ChaosMover/ChaosMovementMode.h
new file mode 100644
--- /dev/null
+++ b/Engine/Plugins/Experimental/ChaosMover/Source/ChaosMover/Public/ChaosMover/ChaosMovementMode.h
@@ -0,0 +1,32 @@
+// Copyright Epic Games, Inc. All Rights Reserved.
+
+#pragma once
+
+#include "MovementMode.h"
+
+#include "ChaosMovementMode.generated.h"
+
+class UChaosMoverSimulation;
+
+/**
+ * Base class for all Chaos movement modes
+ */
+UCLASS(Abstract, Blueprintable, BlueprintType, EditInlineNew, DefaultToInstanced)
+class CHAOSMOVER_API UChaosMovementMode : public UBaseMovementMode
+{
+	GENERATED_BODY()
+
+public:
+	UChaosMovementMode(const FObjectInitializer& ObjectInitializer);
+
+	UFUNCTION(BlueprintPure, Category = ChaosMover)
+	const UChaosMoverSimulation* GetSimulation() const
+	{
+		return Simulation;
+	}
+
+	void SetSimulation(UChaosMoverSimulation* InSimulation);
+
+protected:
+	UChaosMoverSimulation* Simulation;
+};
\ No newline at end of file
diff --git a/Engine/Plugins/Experimental/ChaosMover/Source/ChaosMover/Public/ChaosMover/ChaosMovementModeTransition.h b/Engine/Plugins/Experimental/ChaosMover/Source/ChaosMover/Public/ChaosMover/ChaosMovementModeTransition.h
new file mode 100644
--- /dev/null
+++ b/Engine/Plugins/Experimental/ChaosMover/Source/ChaosMover/Public/ChaosMover/ChaosMovementModeTransition.h
@@ -0,0 +1,35 @@
+// Copyright Epic Games, Inc. All Rights Reserved.
+
+#pragma once
+
+#include "MovementModeTransition.h"
+
+#include "ChaosMovementModeTransition.generated.h"
+
+class UChaosMoverSimulation;
+
+/**
+ * Base class for all Chaos movement mode transitions
+ */
+UCLASS(Abstract, Blueprintable, BlueprintType, EditInlineNew, DefaultToInstanced)
+class CHAOSMOVER_API UChaosMovementModeTransition : public UBaseMovementModeTransition
+{
+	GENERATED_BODY()
+
+public:
+	UChaosMovementModeTransition(const FObjectInitializer& ObjectInitializer);
+
+	UFUNCTION(BlueprintPure, Category = ChaosMover)
+	const UChaosMoverSimulation* GetSimulation() const
+	{
+		return Simulation;
+	}
+
+	void SetSimulation(UChaosMoverSimulation* InSimulation)
+	{
+		Simulation = InSimulation;
+	}
+
+protected:
+	UChaosMoverSimulation* Simulation;
+};
\ No newline at end of file
diff --git a/Engine/Plugins/Experimental/ChaosMover/Source/ChaosMover/Public/ChaosMover/ChaosMoverDeveloperSettings.h b/Engine/Plugins/Experimental/ChaosMover/Source/ChaosMover/Public/ChaosMover/ChaosMoverDeveloperSettings.h
--- a/Engine/Plugins/Experimental/ChaosMover/Source/ChaosMover/Public/ChaosMover/ChaosMoverDeveloperSettings.h
+++ b/Engine/Plugins/Experimental/ChaosMover/Source/ChaosMover/Public/ChaosMover/ChaosMoverDeveloperSettings.h
@@ -11,6 +11,7 @@ namespace UE::ChaosMover
 		CHAOSMOVER_API extern bool bForceSingleThreadedGT;
 		CHAOSMOVER_API extern bool bForceSingleThreadedPT;
 		CHAOSMOVER_API extern bool bDrawGroundQueries;
+		CHAOSMOVER_API extern bool bSkipGenerateMoveIfOverridden;
 	}
 };
 
diff --git a/Engine/Plugins/Experimental/ChaosMover/Source/ChaosMover/Public/ChaosMover/ChaosMoverSimulation.h b/Engine/Plugins/Experimental/ChaosMover/Source/ChaosMover/Public/ChaosMover/ChaosMoverSimulation.h
--- a/Engine/Plugins/Experimental/ChaosMover/Source/ChaosMover/Public/ChaosMover/ChaosMoverSimulation.h
+++ b/Engine/Plugins/Experimental/ChaosMover/Source/ChaosMover/Public/ChaosMover/ChaosMoverSimulation.h
@@ -2,54 +2,114 @@
 
 #pragma once
 
+#include "Chaos/Declares.h"
 #include "ChaosMover/ChaosMoverSimulationTypes.h"
-#include "UObject/WeakObjectPtrTemplates.h"
 #include "ChaosMoverStateMachine.h"
 #include "MoveLibrary/MovementMixer.h"
 #include "MoveLibrary/MoverBlackboard.h"
+#include "UObject/WeakObjectPtrTemplates.h"
+
+#include "ChaosMoverSimulation.generated.h"
 
 namespace Chaos
 {
+	class FCharacterGroundConstraintProxy;
 	class FCollisionContactModifier;
 }
 
 namespace UE::ChaosMover
 {
-	class CHAOSMOVER_API FSimulation
+	struct FMovementModeChangeEvent;
+}
+
+UCLASS(BlueprintType)
+class CHAOSMOVER_API UChaosMoverSimulation : public UObject
+{
+	GENERATED_BODY()
+
+public:
+
+	UChaosMoverSimulation();
+
+	UFUNCTION(BlueprintPure, Category = Mover)
+	const UMoverBlackboard* GetBlackboard() const
+	{
+		return Blackboard;
+	}
+
+	UFUNCTION(BlueprintPure, Category = Mover)
+	UMoverBlackboard* GetBlackboard_Mutable()
 	{
-	public:
-		FSimulation() = default;
-		FSimulation(const FSimulation&) = delete;
-		const FSimulation& operator=(const FSimulation&) = delete;
-		virtual ~FSimulation();
-
-		void Init(const FMoverSyncState& InitialSyncState);
-		void SimulationTick(const FMoverTimeStep& TimeStep, const FSimulationInputData& InputData, FSimulationOutputData& OutputData);
-		void ModifyContacts(Chaos::FCollisionContactModifier& Modifier);
-
-		void InitNetInputData(const FMoverInputCmdContext& InNetInputCmd);
-		void ApplyNetInputData(const FMoverInputCmdContext& InNetInputCmd);
-		void BuildNetInputData(FMoverInputCmdContext& OutNetInputCmd) const;
-		void ApplyNetStateData(const FMoverSyncState& InNetSyncState);
-		void BuildNetStateData(FMoverSyncState& OutNetSyncState) const;
-
-	protected:
-		virtual void OnInit();
-		virtual void OnPreSimulationTick(const FMoverTimeStep& TimeStep, const FSimulationInputData& InputData);
-		virtual void OnSimulationTick(const FMoverTimeStep& TimeStep, const FSimulationInputData& InputData, FSimulationOutputData& OutputData);
-		virtual void OnPostSimulationTick(const FMoverTimeStep& TimeStep, const FSimulationOutputData& OutputData);
-		virtual void OnModifyContacts(Chaos::FCollisionContactModifier& Modifier);
-
-		FMoverSyncState SyncState;
-
-	public:
-		FMoverStateMachine MoverStateMachine;
-		TWeakObjectPtr<UMoverBlackboard> SimBlackboardWeakPtr;
-		TWeakObjectPtr<UMovementMixer> MovementMixerWeakPtr;
-
-	private:
-		FMoverInputCmdContext InputCmd;
-		bool bInputCmdOverridden = false;
+		return Blackboard;
+	}
+
+	UFUNCTION(BlueprintPure, Category = Mover)
+	const UBaseMovementMode* GetCurrentMovementMode() const;
+
+	UFUNCTION(BlueprintPure, Category = Mover)
+	const UBaseMovementMode* FindMovementModeByName(const FName& Name) const;
+
+	struct FInitParams
+	{
+		TMap<FName, TWeakObjectPtr<UBaseMovementMode>> ModesToRegister;
+		TArray<TWeakObjectPtr<UBaseMovementModeTransition>> TransitionsToRegister;
+		FMoverSyncState InitialSyncState;
+		FName StartingMovementMode = NAME_None;
+		TWeakObjectPtr<UNullMovementMode> NullMovementMode = nullptr;
+		TWeakObjectPtr<UImmediateMovementModeTransition> ImmediateModeTransition = nullptr;
+		TWeakObjectPtr<UMovementMixer> MovementMixer = nullptr;
+		Chaos::FCharacterGroundConstraintProxy* CharacterConstraintProxy = nullptr;
+		Chaos::FConstPhysicsObjectHandle PhysicsObject = nullptr;
+		Chaos::FPhysicsSolver* Solver = nullptr;
+		UWorld* World = nullptr;
 	};
 
-} // namespace UE::ChaosMover
\ No newline at end of file
+	void Init(const FInitParams& InitParams);
+	void Deinit();
+	void SimulationTick(const FMoverTimeStep& TimeStep, const UE::ChaosMover::FSimulationInputData& InputData, UE::ChaosMover::FSimulationOutputData& OutputData);
+	void ModifyContacts(Chaos::FCollisionContactModifier& Modifier) const;
+
+	const TArray<UE::ChaosMover::FMovementModeChangeEvent>& GetQueuedModeChangeEvents() const;
+	void ClearQueuedModeChangeEvents();
+
+	void InitNetInputData(const FMoverInputCmdContext& InNetInputCmd);
+	void ApplyNetInputData(const FMoverInputCmdContext& InNetInputCmd);
+	void BuildNetInputData(FMoverInputCmdContext& OutNetInputCmd) const;
+	void ApplyNetStateData(const FMoverSyncState& InNetSyncState);
+	void BuildNetStateData(FMoverSyncState& OutNetSyncState) const;
+
+protected:
+	virtual void OnInit();
+	virtual void OnDeinit();
+	virtual void OnPreSimulationTick(const FMoverTimeStep& TimeStep, const UE::ChaosMover::FSimulationInputData& InputData);
+	virtual void OnSimulationTick(const FMoverTimeStep& TimeStep, const UE::ChaosMover::FSimulationInputData& InputData, UE::ChaosMover::FSimulationOutputData& OutputData);
+	virtual void OnPostSimulationTick(const FMoverTimeStep& TimeStep, UE::ChaosMover::FSimulationOutputData& OutputData);
+	virtual void OnModifyContacts(Chaos::FCollisionContactModifier& Modifier) const;
+
+	virtual void PreSimulationTickCharacter(const IChaosCharacterMovementModeInterface& CharacterMode, const FMoverTimeStep& TimeStep, const UE::ChaosMover::FSimulationInputData& InputData);
+	virtual void PostSimulationTickCharacter(const IChaosCharacterMovementModeInterface& CharacterMode, const FMoverTimeStep& TimeStep, UE::ChaosMover::FSimulationOutputData& OutputData);
+	virtual void PostSimulationTickCharacterConstraint(const IChaosCharacterConstraintMovementModeInterface& CharacterConstraintMode, const FMoverTimeStep& TimeStep, UE::ChaosMover::FSimulationOutputData& OutputData);
+
+	Chaos::FPBDRigidParticleHandle* GetControlledParticle() const;
+	void EnableCharacterConstraint();
+	void DisableCharacterConstraint();
+
+	FMoverSyncState CurrentSyncState;
+	FVector DeltaPosition;
+	FRotator DeltaRotation;
+
+	UE::ChaosMover::FMoverStateMachine StateMachine;
+
+	UPROPERTY(Transient)
+	TObjectPtr<UMoverBlackboard> Blackboard = nullptr;
+
+	TWeakObjectPtr<UMovementMixer> MovementMixerWeakPtr = nullptr;
+	Chaos::FCharacterGroundConstraintProxy* CharacterConstraintProxy = nullptr;
+	Chaos::FConstPhysicsObjectHandle PhysicsObject = nullptr;
+	Chaos::FPhysicsSolver* Solver = nullptr;
+	UWorld* World = nullptr;
+
+private:
+	FMoverInputCmdContext InputCmd;
+	bool bInputCmdOverridden = false;
+};
diff --git a/Engine/Plugins/Experimental/ChaosMover/Source/ChaosMover/Public/ChaosMover/ChaosMoverSimulationTypes.h b/Engine/Plugins/Experimental/ChaosMover/Source/ChaosMover/Public/ChaosMover/ChaosMoverSimulationTypes.h
--- a/Engine/Plugins/Experimental/ChaosMover/Source/ChaosMover/Public/ChaosMover/ChaosMoverSimulationTypes.h
+++ b/Engine/Plugins/Experimental/ChaosMover/Source/ChaosMover/Public/ChaosMover/ChaosMoverSimulationTypes.h
@@ -8,6 +8,13 @@
 
 #include "ChaosMoverSimulationTypes.generated.h"
 
+namespace Chaos
+{
+	class FCharacterGroundConstraintSettings;
+}
+
+struct FConstraintProfileProperties;
+
 namespace UE::ChaosMover
 {
 	struct FSimulationInputData
@@ -15,11 +22,11 @@ namespace UE::ChaosMover
 		void Reset()
 		{
 			InputCmd.Reset();
-			InputParams.Empty();
+			LocalInputParams.AuxStateCollection.Empty();
 		}
 
 		mutable FMoverInputCmdContext InputCmd; // can be overridden by network physics
-		FMoverDataCollection InputParams;
+		FMoverAuxStateContext LocalInputParams;
 	};
 
 	using FSimulationOutputData = ::UE::Mover::FSimulationOutputData;
@@ -53,11 +60,69 @@ struct CHAOSMOVER_API FChaosMoverSimulationDefaultInputs : public FMoverDataStru
 
 	FCollisionResponseParams CollisionResponseParams;
 	FCollisionQueryParams CollisionQueryParams;
+
+	UPROPERTY(BlueprintReadOnly, Category=ChaosMover)
 	FVector UpDir;
+
 	FVector Gravity;
 	Chaos::FPhysicsObjectHandle PhysicsObject;
 	AActor* OwningActor;
-	ECollisionChannel CollisionChannel;
 	float BoundingCylinderRadius;
 	float BoundingCylinderfHeight;
+	float PhysicsObjectGravity;
+	ECollisionChannel CollisionChannel;
+};
+
+UINTERFACE(BlueprintType, MinimalAPI, meta = (CannotImplementInterfaceInBlueprint))
+class UChaosCharacterMovementModeInterface : public UInterface
+{
+	GENERATED_BODY()
+};
+
+class IChaosCharacterMovementModeInterface
+{
+	GENERATED_BODY()
+
+public:
+	UFUNCTION(BlueprintCallable, Category=ChaosMover)
+	virtual float GetTargetHeight() const = 0;
+
+	UFUNCTION(BlueprintCallable, Category = ChaosMover)
+	virtual float GetGroundQueryRadius() const = 0;
+
+	UFUNCTION(BlueprintCallable, Category = ChaosMover)
+	virtual float GetMaxWalkSlopeCosine() const = 0;
+
+	UFUNCTION(BlueprintCallable, Category = ChaosMover)
+	virtual bool ShouldCharacterRemainUpright() const = 0;
+};
+
+UINTERFACE(MinimalAPI)
+class UChaosCharacterConstraintMovementModeInterface : public UInterface
+{
+	GENERATED_BODY()
+};
+
+class IChaosCharacterConstraintMovementModeInterface
+{
+	GENERATED_BODY()
+
+public:
+	virtual bool ShouldEnableConstraint() const = 0;
+	virtual  void UpdateConstraintSettings(Chaos::FCharacterGroundConstraintSettings& ConstraintSettings) const = 0;
+};
+
+UINTERFACE(MinimalAPI)
+class UChaosConstraintMovementModeInterface : public UInterface
+{
+	GENERATED_BODY()
+};
+
+class IChaosConstraintMovementModeInterface
+{
+	GENERATED_BODY()
+
+public:
+	virtual bool ShouldEnableConstraint() const = 0;
+	virtual const FConstraintProfileProperties& GetConstraintProperties() const = 0;
 };
diff --git a/Engine/Plugins/Experimental/ChaosMover/Source/ChaosMover/Public/ChaosMover/Character/Modes/ChaosCharacterMovementMode.h b/Engine/Plugins/Experimental/ChaosMover/Source/ChaosMover/Public/ChaosMover/Character/Modes/ChaosCharacterMovementMode.h
--- a/Engine/Plugins/Experimental/ChaosMover/Source/ChaosMover/Public/ChaosMover/Character/Modes/ChaosCharacterMovementMode.h
+++ b/Engine/Plugins/Experimental/ChaosMover/Source/ChaosMover/Public/ChaosMover/Character/Modes/ChaosCharacterMovementMode.h
@@ -2,7 +2,8 @@
 
 #pragma once
 
-#include "MovementMode.h"
+#include "ChaosMover/ChaosMovementMode.h"
+#include "ChaosMover/ChaosMoverSimulationTypes.h"
 
 #include "ChaosCharacterMovementMode.generated.h"
 
@@ -14,8 +15,10 @@ namespace Chaos
 /**
  * Base class for all Chaos character movement modes
  */
-UCLASS(Blueprintable, BlueprintType, EditInlineNew, DefaultToInstanced)
-class CHAOSMOVER_API UChaosCharacterMovementMode : public UBaseMovementMode
+UCLASS(Abstract, Blueprintable, BlueprintType, EditInlineNew, DefaultToInstanced)
+class CHAOSMOVER_API UChaosCharacterMovementMode :	public UChaosMovementMode,
+													public IChaosCharacterMovementModeInterface,
+													public IChaosCharacterConstraintMovementModeInterface
 {
 	GENERATED_BODY()
 
@@ -25,33 +28,41 @@ public:
 	virtual void OnRegistered(const FName ModeName) override;
 	virtual void OnUnregistered() override;
 
-	// Update the constraint settings with the data defined by this mode
-	virtual void UpdateConstraintSettings(Chaos::FCharacterGroundConstraintSettings& ConstraintSettings) const;
+	// IChaosCharacterConstraintMovementModeInterface
+	virtual bool ShouldEnableConstraint() const override
+	{
+		return true;
+	}
 
-	float GetTargetHeight() const
+	virtual void UpdateConstraintSettings(Chaos::FCharacterGroundConstraintSettings& ConstraintSettings) const override;
+
+	// IChaosCharacterMovementModeInterface
+	virtual float GetTargetHeight() const override
 	{
 		return TargetHeight;
 	}
 
-	void SetTargetHeightOverride(float InTargetHeight);
-	void ClearTargetHeightOverride();
-
-	float GetQueryRadius() const
+	virtual float GetGroundQueryRadius() const override
 	{
 		return QueryRadius;
 	}
 
-	void SetQueryRadiusOverride(float InQueryRadius);
-	void ClearQueryRadiusOverride();
+	virtual float GetMaxWalkSlopeCosine() const override;
 
-	bool ShouldCharacterRemainUpright() const
+	virtual bool ShouldCharacterRemainUpright() const override
 	{
 		return bShouldCharacterRemainUpright;
 	}
 
+protected:
+	void SetTargetHeightOverride(float InTargetHeight);
+	void ClearTargetHeightOverride();
+
+	void SetQueryRadiusOverride(float InQueryRadius);
+	void ClearQueryRadiusOverride();
+
 	TObjectPtr<const class USharedChaosCharacterMovementSettings> SharedSettings;
 
-protected:
 	// Maximum force the character can apply to reach the motion target
 	UPROPERTY(EditAnywhere, Category = "Constraint Settings", meta = (ClampMin = "0", UIMin = "0", ForceUnits = "Newtons"))
 	float RadialForceLimit = 1500.0f;
diff --git a/Engine/Plugins/Experimental/ChaosMover/Source/ChaosMover/Public/ChaosMover/Character/Modes/ChaosFallingMode.h b/Engine/Plugins/Experimental/ChaosMover/Source/ChaosMover/Public/ChaosMover/Character/Modes/ChaosFallingMode.h
--- a/Engine/Plugins/Experimental/ChaosMover/Source/ChaosMover/Public/ChaosMover/Character/Modes/ChaosFallingMode.h
+++ b/Engine/Plugins/Experimental/ChaosMover/Source/ChaosMover/Public/ChaosMover/Character/Modes/ChaosFallingMode.h
@@ -19,4 +19,55 @@ public:
 
 	virtual void GenerateMove_Implementation(const FMoverTickStartData& StartState, const FMoverTimeStep& TimeStep, FProposedMove& OutProposedMove) const override;
 	virtual void SimulationTick_Implementation(const FSimulationTickParams& Params, FMoverTickEndData& OutputState) override;
+
+protected:
+	/**
+	 * If true, actor will land and lose all speed in the vertical direction upon landing. If false, actor's vertical speed will be redirected based on the surface normal it hit.
+	 * Note: Actor's horizontal speed will not be affected if true. If false, horizontal speed may be increased on landing.
+	 */
+	UPROPERTY(Category = Mover, EditAnywhere, BlueprintReadWrite)
+	bool bCancelVerticalSpeedOnLanding;
+
+	/**
+	 * When falling, amount of movement control available to the actor.
+	 * 0 = no control, 1 = full control
+	 */
+	UPROPERTY(Category = Mover, EditAnywhere, BlueprintReadWrite, meta = (ClampMin = "0", ClampMax = "1.0"))
+	float AirControlPercentage;
+
+	/**
+	 * Deceleration to apply to air movement when falling slower than terminal velocity.
+	 * Note: This is NOT applied to vertical velocity, only movement plane velocity
+	 */
+	UPROPERTY(Category = Mover, EditAnywhere, BlueprintReadWrite, meta = (ClampMin = "0", ForceUnits = "cm/s^2"))
+	float FallingDeceleration;
+
+	/**
+	 * Deceleration to apply to air movement when falling faster than terminal velocity
+	 * Note: This is NOT applied to vertical velocity, only movement plane velocity
+	 */
+	UPROPERTY(Category = Mover, EditAnywhere, BlueprintReadWrite, meta = (ClampMin = "0", ForceUnits = "cm/s^2"))
+	float OverTerminalSpeedFallingDeceleration;
+
+	/**
+	 * If the actor's movement plane velocity is greater than this speed falling will start applying OverTerminalSpeedFallingDeceleration instead of FallingDeceleration
+	 * The expected behavior is to set OverTerminalSpeedFallingDeceleration higher than FallingDeceleration so the actor will slow down faster
+	 * when going over TerminalMovementPlaneSpeed.
+	 */
+	UPROPERTY(Category = Mover, EditAnywhere, BlueprintReadWrite, meta = (ClampMin = "0", ForceUnits = "cm/s"))
+	float TerminalMovementPlaneSpeed;
+
+	/** When exceeding maximum vertical speed, should it be enforced via a hard clamp? If false, VerticalFallingDeceleration will be used for a smoother transition to the terminal speed limit. */
+	UPROPERTY(Category = Mover, EditAnywhere, BlueprintReadWrite)
+	bool bShouldClampTerminalVerticalSpeed;
+
+	/** Deceleration to apply to vertical velocity when it's greater than TerminalVerticalSpeed. Only used if bShouldClampTerminalVerticalSpeed is false. */
+	UPROPERTY(Category = Mover, EditAnywhere, BlueprintReadWrite, meta = (EditCondition = "!bShouldClampTerminalVerticalSpeed", ClampMin = "0", ForceUnits = "cm/s^2"))
+	float VerticalFallingDeceleration;
+
+	/**
+	 * If the actors vertical velocity is greater than this speed VerticalFallingDeceleration will be applied to vertical velocity
+	 */
+	UPROPERTY(Category = Mover, EditAnywhere, BlueprintReadWrite, meta = (ClampMin = "0", ForceUnits = "cm/s"))
+	float TerminalVerticalSpeed;
 };
\ No newline at end of file
diff --git a/Engine/Plugins/Experimental/ChaosMover/Source/ChaosMover/Public/ChaosMover/Character/Settings/SharedChaosCharacterMovementSettings.h b/Engine/Plugins/Experimental/ChaosMover/Source/ChaosMover/Public/ChaosMover/Character/Settings/SharedChaosCharacterMovementSettings.h
--- a/Engine/Plugins/Experimental/ChaosMover/Source/ChaosMover/Public/ChaosMover/Character/Settings/SharedChaosCharacterMovementSettings.h
+++ b/Engine/Plugins/Experimental/ChaosMover/Source/ChaosMover/Public/ChaosMover/Character/Settings/SharedChaosCharacterMovementSettings.h
@@ -3,6 +3,7 @@
 #pragma once
 
 #include "MovementMode.h"
+
 #include "SharedChaosCharacterMovementSettings.generated.h"
 
 /**
@@ -24,6 +25,34 @@ public:
 		return MaxWalkSlopeCosine;
 	}
 
+	/** Default max linear rate of deceleration when there is no controlled input */
+	UPROPERTY(BlueprintReadWrite, EditAnywhere, Category = "General", meta = (ClampMin = "0", UIMin = "0", ForceUnits = "cm/s^2"))
+	float Deceleration = 4000.f;
+
+	/** Default max linear rate of acceleration for controlled input. May be scaled based on magnitude of input. */
+	UPROPERTY(BlueprintReadWrite, EditAnywhere, Category = "General", meta = (ClampMin = "0", UIMin = "0", ForceUnits = "cm/s^2"))
+	float Acceleration = 4000.f;
+
+	/** Maximum rate of turning rotation (degrees per second). Negative numbers indicate instant rotation and should cause rotation to snap instantly to desired direction. */
+	UPROPERTY(BlueprintReadWrite, EditAnywhere, Category = "General", meta = (ClampMin = "-1", UIMin = "0", ForceUnits = "degrees/s"))
+	float TurningRate = 500.f;
+
+	/** Speeds velocity direction changes while turning, to reduce sliding */
+	UPROPERTY(BlueprintReadWrite, EditAnywhere, Category = "General", meta = (ClampMin = "0", UIMin = "0", ForceUnits = "Multiplier"))
+	float TurningBoost = 8.f;
+
+	/** Maximum speed in the movement plane */
+	UPROPERTY(BlueprintReadWrite, EditAnywhere, Category = "General", meta = (ClampMin = "0", UIMin = "0", ForceUnits = "cm/s"))
+	float MaxSpeed = 800.f;
+
+	/**
+	 * Should use acceleration for velocity based movement intent?
+	 * If true, acceleration is applied when using velocity input to reach the target velocity.
+	 * If false, velocity is set directly, disregarding acceleration.
+	 */
+	UPROPERTY(Category = "General", EditAnywhere, BlueprintReadWrite, AdvancedDisplay)
+	bool bUseAccelerationForVelocityMove = true;
+
 protected:
 	/** Walkable slope angle, represented as cosine(max slope angle) for performance reasons. Ex: for max slope angle of 30 degrees, value is cosine(30 deg) = 0.866 */
 	UPROPERTY(BlueprintReadWrite, EditAnywhere, Category = "Ground Movement", BlueprintSetter=SetMaxWalkableSlopeAngle, meta = (ClampMin = "0", UIMin = "0", ClampMax = "90", UIMax = "90", ForceUnits = "degrees"))
diff --git a/Engine/Plugins/Experimental/ChaosMover/Source/ChaosMover/Public/ChaosMover/Character/Transitions/ChaosCharacterLandingCheck.h b/Engine/Plugins/Experimental/ChaosMover/Source/ChaosMover/Public/ChaosMover/Character/Transitions/ChaosCharacterLandingCheck.h
new file mode 100644
--- /dev/null
+++ b/Engine/Plugins/Experimental/ChaosMover/Source/ChaosMover/Public/ChaosMover/Character/Transitions/ChaosCharacterLandingCheck.h
@@ -0,0 +1,36 @@
+// Copyright Epic Games, Inc. All Rights Reserved.
+
+#pragma once
+
+#include "ChaosMover/ChaosMovementModeTransition.h"
+
+#include "ChaosCharacterLandingCheck.generated.h"
+
+
+UCLASS(Blueprintable, BlueprintType, EditInlineNew, DefaultToInstanced)
+class CHAOSMOVER_API UChaosCharacterLandingCheck : public UChaosMovementModeTransition
+{
+	GENERATED_BODY()
+
+public:
+	UChaosCharacterLandingCheck(const FObjectInitializer& ObjectInitializer);
+
+	virtual FTransitionEvalResult Evaluate_Implementation(const FSimulationTickParams& Params) const override;
+	virtual void Trigger_Implementation(const FSimulationTickParams& Params) override;
+
+	/** Depth at which the pawn starts swimming */
+	UPROPERTY(BlueprintReadWrite, EditAnywhere, Category = Parameters, meta = (Units = "cm"))
+	float SwimmingStartImmersionDepth = 64.5f;
+
+	/** Height at which we consider the character to be on the ground */
+	UPROPERTY(BlueprintReadWrite, EditAnywhere, Category = Parameters, meta = (Units = "cm"))
+	float FloorDistanceTolerance = 0.0f;
+
+	/** Name of movement mode to transition to when landing on ground */
+	UPROPERTY(BlueprintReadWrite, EditAnywhere, Category = Parameters)
+	FName TransitionToGroundMode;
+
+	/** Name of movement mode to transition to when landing on water */
+	UPROPERTY(BlueprintReadWrite, EditAnywhere, Category = Parameters)
+	FName TransitionToWaterMode;
+};
\ No newline at end of file
diff --git a/Engine/Plugins/Experimental/ChaosMover/Source/ChaosMover/Public/ChaosMover/Utilities/ChaosGroundMovementUtils.h b/Engine/Plugins/Experimental/ChaosMover/Source/ChaosMover/Public/ChaosMover/Utilities/ChaosGroundMovementUtils.h
new file mode 100644
--- /dev/null
+++ b/Engine/Plugins/Experimental/ChaosMover/Source/ChaosMover/Public/ChaosMover/Utilities/ChaosGroundMovementUtils.h
@@ -0,0 +1,23 @@
+// Copyright Epic Games, Inc. All Rights Reserved.
+
+#pragma once
+
+#include "Chaos/ParticleHandleFwd.h"
+#include "Kismet/BlueprintFunctionLibrary.h"
+
+#include "ChaosGroundMovementUtils.generated.h"
+
+struct FFloorCheckResult;
+
+UCLASS()
+class CHAOSMOVER_API UChaosGroundMovementUtils : public UBlueprintFunctionLibrary
+{
+	GENERATED_BODY()
+
+public:
+	UFUNCTION(BlueprintCallable, Category = ChaosMover)
+	static FVector ComputeLocalGroundVelocity_Internal(const FVector& Position, const FFloorCheckResult& FloorResult);
+
+protected:
+	static Chaos::FPBDRigidParticleHandle* GetRigidParticleHandleFromFloorResult_Internal(const FFloorCheckResult& FloorResult);
+};
\ No newline at end of file
diff --git a/Engine/Plugins/Experimental/Mover/Source/Mover/Public/ChaosVisualDebugger/MoverCVDRuntimeTrace.h b/Engine/Plugins/Experimental/Mover/Source/Mover/Public/ChaosVisualDebugger/MoverCVDRuntimeTrace.h
--- a/Engine/Plugins/Experimental/Mover/Source/Mover/Public/ChaosVisualDebugger/MoverCVDRuntimeTrace.h
+++ b/Engine/Plugins/Experimental/Mover/Source/Mover/Public/ChaosVisualDebugger/MoverCVDRuntimeTrace.h
@@ -19,7 +19,7 @@ class UMoverComponent;
 class FMoverCVDRuntimeTrace
 {
 public:
-	static void TraceMoverData(UMoverComponent* MoverComponent, const FMoverInputCmdContext* InputCmd, const FMoverSyncState* SyncState);
+	MOVER_API static void TraceMoverData(UMoverComponent* MoverComponent, const FMoverInputCmdContext* InputCmd, const FMoverSyncState* SyncState);
 
 	MOVER_API static void UnwrapSimData(const FMoverCVDSimDataWrapper& InSimDataWrapper, TSharedPtr<FMoverInputCmdContext>& OutInputCmd, TSharedPtr<FMoverSyncState>& OutSyncState);
 	static void WrapSimData(uint32 SolverID, uint32 ParticleID, const FMoverInputCmdContext& InInputCmd, const FMoverSyncState& InSyncState, FMoverCVDSimDataWrapper& OutSimDataWrapper);
@@ -36,7 +36,7 @@ CVD_DECLARE_OPTIONAL_DATA_CHANNEL_EXTERN(MoverData, MOVER_API)
 class FMoverCVDRuntimeTrace
 {
 public:
-	static void TraceMoverData(UMoverComponent* MoverComponent, const FMoverInputCmdContext* InputCmd, const FMoverSyncState* SyncState) {}
+	MOVER_API static void TraceMoverData(UMoverComponent* MoverComponent, const FMoverInputCmdContext* InputCmd, const FMoverSyncState* SyncState) {}
 	MOVER_API static void UnwrapSimData(const FMoverCVDSimDataWrapper& InSimDataWrapper, TSharedPtr<FMoverInputCmdContext>& OutInputCmd, TSharedPtr<FMoverSyncState>& OutSyncState) {}
 	static void WrapSimData(uint32 SolverID, uint32 ParticleID, const FMoverInputCmdContext& InInputCmd, const FMoverSyncState& InSyncState, FMoverCVDSimDataWrapper& OutSimDataWrapper) {}
 };
