ChaosMover - Get walking mode working and add falling check transition

diff --git a/Engine/Plugins/Experimental/ChaosMover/Source/ChaosMover/Private/Backends/ChaosMoverBackend.cpp b/Engine/Plugins/Experimental/ChaosMover/Source/ChaosMover/Private/Backends/ChaosMoverBackend.cpp
--- a/Engine/Plugins/Experimental/ChaosMover/Source/ChaosMover/Private/Backends/ChaosMoverBackend.cpp
+++ b/Engine/Plugins/Experimental/ChaosMover/Source/ChaosMover/Private/Backends/ChaosMoverBackend.cpp
@@ -335,10 +335,22 @@ void UChaosMoverBackendComponent::ProduceInputData(int32 PhysicsStep, int32 NumS
 	SimInputs.Gravity = MoverComp.GetGravityAcceleration();
 	SimInputs.UpDir = MoverComp.GetUpDirection();
 	SimInputs.OwningActor = GetOwner();
+	SimInputs.World = GetWorld();
+
 	if (UPrimitiveComponent* PrimComp = Cast<UPrimitiveComponent>(MoverComp.GetUpdatedComponent()))
 	{
+		SimInputs.CollisionQueryParams = FCollisionQueryParams(SCENE_QUERY_STAT(ChaosMoverQuery), false, PrimComp->GetOwner());
+		SimInputs.CollisionQueryParams.bTraceIntoSubComponents = false;
+		SimInputs.CollisionResponseParams = FCollisionResponseParams(ECR_Overlap);
+		SimInputs.CollisionResponseParams.CollisionResponse.SetResponse(ECC_WorldStatic, ECR_Block);
+		SimInputs.CollisionResponseParams.CollisionResponse.SetResponse(ECC_WorldDynamic, ECR_Block);
+		SimInputs.CollisionResponseParams.CollisionResponse.SetResponse(ECC_Vehicle, ECR_Block);
+		SimInputs.CollisionResponseParams.CollisionResponse.SetResponse(ECC_Destructible, ECR_Block);
+		SimInputs.CollisionResponseParams.CollisionResponse.SetResponse(ECC_PhysicsBody, ECR_Block);
 		PrimComp->InitSweepCollisionParams(SimInputs.CollisionQueryParams, SimInputs.CollisionResponseParams);
+
 		SimInputs.CollisionChannel = PrimComp->GetCollisionObjectType();
+		PrimComp->CalcBoundingCylinder(SimInputs.PawnCollisionRadius, SimInputs.PawnCollisionHalfHeight);
 	}
 	if (IPhysicsComponent* PhysComp = Cast<IPhysicsComponent>(MoverComp.GetUpdatedComponent()))
 	{
diff --git a/Engine/Plugins/Experimental/ChaosMover/Source/ChaosMover/Private/ChaosMoverSimulation.cpp b/Engine/Plugins/Experimental/ChaosMover/Source/ChaosMover/Private/ChaosMoverSimulation.cpp
--- a/Engine/Plugins/Experimental/ChaosMover/Source/ChaosMover/Private/ChaosMoverSimulation.cpp
+++ b/Engine/Plugins/Experimental/ChaosMover/Source/ChaosMover/Private/ChaosMoverSimulation.cpp
@@ -144,37 +144,6 @@ void UChaosMoverSimulation::ModifyContacts(Chaos::FCollisionContactModifier& Mod
 	OnModifyContacts(Modifier);
 }
 
-void UChaosMoverSimulation::PreSimulationTickCharacter(const IChaosCharacterMovementModeInterface& CharacterMode, const FMoverTimeStep& TimeStep, const UE::ChaosMover::FSimulationInputData& InputData)
-{
-	const FMoverDefaultSyncState* DefaultSyncState = CurrentSyncState.SyncStateCollection.FindDataByType<FMoverDefaultSyncState>();
-	check(DefaultSyncState);
-	check(Blackboard.Get());
-
-	if (const FChaosMoverSimulationDefaultInputs* DefaultInputs = InputData.LocalInputParams.AuxStateCollection.FindDataByType<FChaosMoverSimulationDefaultInputs>())
-	{
-		UE::ChaosMover::Utils::FFloorSweepParams SweepParams{
-			.ResponseParams = DefaultInputs->CollisionResponseParams,
-			.QueryParams = DefaultInputs->CollisionQueryParams,
-			.Location = DefaultSyncState->GetLocation_WorldSpace(),
-			.DeltaPos = DefaultSyncState->GetVelocity_WorldSpace() * TimeStep.StepMs * 0.001f,
-			.UpDir = DefaultInputs->UpDir,
-			.World = World,
-			.QueryDistance = 1.2f * CharacterMode.GetTargetHeight(),
-			.QueryRadius = CharacterMode.GetGroundQueryRadius(),
-			.MaxWalkSlopeCosine = CharacterMode.GetMaxWalkSlopeCosine(),
-			.TargetHeight = CharacterMode.GetTargetHeight(),
-			.CollisionChannel = DefaultInputs->CollisionChannel
-		};
-
-		FFloorCheckResult FloorResult;
-		FWaterCheckResult WaterResult;
-		UE::ChaosMover::Utils::FloorSweep_Internal(SweepParams, FloorResult, WaterResult);
-
-		Blackboard->Set(CommonBlackboard::LastFloorResult, FloorResult);
-		Blackboard->Set(CommonBlackboard::LastWaterResult, WaterResult);
-	}
-}
-
 void UChaosMoverSimulation::OnPreSimulationTick(const FMoverTimeStep& TimeStep, const UE::ChaosMover::FSimulationInputData& InputData)
 {
 	if (bInputCmdOverridden)
@@ -198,15 +167,6 @@ void UChaosMoverSimulation::OnPreSimulationTick(const FMoverTimeStep& TimeStep,
 	{
 		DefaultSyncState.SetTransforms_WorldSpace(ParticleHandle->GetX(), FRotator(ParticleHandle->GetR()), ParticleHandle->GetV());
 	}
-
-	// TODO - make this more extensible
-	if (TStrongObjectPtr<UBaseMovementMode> CurrentModePtr = StateMachine.GetCurrentMode().Pin())
-	{
-		if (IChaosCharacterMovementModeInterface* CharacterMode = Cast<IChaosCharacterMovementModeInterface>(CurrentModePtr.Get()))
-		{
-			PreSimulationTickCharacter(*CharacterMode, TimeStep, InputData);
-		}
-	}
 }
 
 void UChaosMoverSimulation::OnSimulationTick(const FMoverTimeStep& TimeStep, const UE::ChaosMover::FSimulationInputData& InputData, UE::ChaosMover::FSimulationOutputData& OutputData)
@@ -215,6 +175,7 @@ void UChaosMoverSimulation::OnSimulationTick(const FMoverTimeStep& TimeStep, con
 
 	FMoverTickStartData TickStartData(InputData.InputCmd, CurrentSyncState, InputData.LocalInputParams);
 	FMoverTickEndData TickEndData(&CurrentSyncState, &InputData.LocalInputParams);
+
 	StateMachine.OnSimulationTick(TimeStep, TickStartData, Blackboard.Get(), MovementMixerWeakPtr.Get(), TickEndData);
 
 	// Copy the sync state locally and to the output data
@@ -227,16 +188,17 @@ void UChaosMoverSimulation::OnSimulationTick(const FMoverTimeStep& TimeStep, con
 		const FMoverDefaultSyncState* DefaultSyncState = OutputData.SyncState.SyncStateCollection.FindDataByType<FMoverDefaultSyncState>();
 		check(DefaultSyncState);
 
-		FRotator TmpRotation = DefaultSyncState->GetOrientation_WorldSpace() - FRotator(ParticleHandle->GetR());
-		FRotator Winding;
-		TmpRotation.GetWindingAndRemainder(Winding, DeltaRotation);
+		FRotator DeltaRotation = DefaultSyncState->GetOrientation_WorldSpace() - FRotator(ParticleHandle->GetR());
+		FRotator Winding, Remainder;
+		DeltaRotation.GetWindingAndRemainder(Winding, Remainder);
 
-		DeltaPosition = DefaultSyncState->GetLocation_WorldSpace() - ParticleHandle->GetX();
+		TargetDeltaPosition = DefaultSyncState->GetLocation_WorldSpace() - ParticleHandle->GetX();
+		TargetDeltaFacing = FMath::DegreesToRadians(Remainder.Yaw);
 	}
 	else
 	{
-		DeltaPosition = FVector::ZeroVector;
-		DeltaRotation = FRotator::ZeroRotator;
+		TargetDeltaPosition = FVector::ZeroVector;
+		TargetDeltaFacing = 0.0f;
 	}
 }
 
@@ -249,8 +211,6 @@ void UChaosMoverSimulation::PostSimulationTickCharacter(const IChaosCharacterMov
 	if (Chaos::FPBDRigidParticleHandle* ParticleHandle = UChaosMoverSimulation::GetControlledParticle())
 	{
 		ParticleHandle->SetV(DefaultSyncState->GetVelocity_WorldSpace());
-
-		float TargetDeltaFacing = FMath::DegreesToRadians(DeltaRotation.Yaw);
 		const float DeltaTimeSeconds = TimeStep.StepMs * 0.001f;
 		if (DeltaTimeSeconds > UE_SMALL_NUMBER)
 		{
@@ -340,8 +300,8 @@ void UChaosMoverSimulation::PostSimulationTickCharacterConstraint(const IChaosCh
 
 		ConstraintHandle->SetData({
 			FloorResult.HitResult.ImpactNormal,
-			DeltaPosition,
-			DeltaRotation.Yaw,
+			TargetDeltaPosition,
+			TargetDeltaFacing,
 			FloorResult.FloorDist,
 			WalkableSlopeCosine
 			});
diff --git a/Engine/Plugins/Experimental/ChaosMover/Source/ChaosMover/Private/ChaosMoverSimulationTypes.cpp b/Engine/Plugins/Experimental/ChaosMover/Source/ChaosMover/Private/ChaosMoverSimulationTypes.cpp
--- a/Engine/Plugins/Experimental/ChaosMover/Source/ChaosMover/Private/ChaosMoverSimulationTypes.cpp
+++ b/Engine/Plugins/Experimental/ChaosMover/Source/ChaosMover/Private/ChaosMoverSimulationTypes.cpp
@@ -9,7 +9,10 @@ void FChaosMoverSimulationDefaultInputs::Reset()
 	UpDir = FVector::UpVector;
 	Gravity = -980.7 * UpDir;
 	PhysicsObjectGravity = 0.0f;
+	PawnCollisionHalfHeight = 40.0f;
+	PawnCollisionRadius = 30.0f;
 	PhysicsObject = nullptr;
 	OwningActor = nullptr;
+	World = nullptr;
 	CollisionChannel = ECC_Pawn;
 }
\ No newline at end of file
diff --git a/Engine/Plugins/Experimental/ChaosMover/Source/ChaosMover/Private/Character/Modes/ChaosFallingMode.cpp b/Engine/Plugins/Experimental/ChaosMover/Source/ChaosMover/Private/Character/Modes/ChaosFallingMode.cpp
--- a/Engine/Plugins/Experimental/ChaosMover/Source/ChaosMover/Private/Character/Modes/ChaosFallingMode.cpp
+++ b/Engine/Plugins/Experimental/ChaosMover/Source/ChaosMover/Private/Character/Modes/ChaosFallingMode.cpp
@@ -5,9 +5,12 @@
 #include "ChaosMover/ChaosMoverLog.h"
 #include "ChaosMover/ChaosMoverSimulation.h"
 #include "ChaosMover/Character/Settings/SharedChaosCharacterMovementSettings.h"
+#include "ChaosMover/Character/Transitions/ChaosCharacterLandingCheck.h"
+#include "ChaosMover/Utilities/ChaosMoverQueryUtils.h"
 #include "MoveLibrary/AirMovementUtils.h"
 #include "MoveLibrary/FloorQueryUtils.h"
 #include "MoveLibrary/MovementUtils.h"
+#include "MoveLibrary/WaterMovementUtils.h"
 
 #include UE_INLINE_GENERATED_CPP_BY_NAME(ChaosFallingMode)
 
@@ -27,6 +30,12 @@ UChaosFallingMode::UChaosFallingMode(const FObjectInitializer& ObjectInitializer
 	GameplayTags.AddTag(Mover_IsInAir);
 	GameplayTags.AddTag(Mover_IsFalling);
 	GameplayTags.AddTag(Mover_SkipAnimRootMotion);
+
+	RadialForceLimit = 0.0f;
+	SwingTorqueLimit = 3000.0f;
+	TwistTorqueLimit = 0.0f;
+
+	Transitions.Add(CreateDefaultSubobject<UChaosCharacterLandingCheck>(TEXT("DefaultLandingCheck")));
 }
 
 void UChaosFallingMode::GenerateMove_Implementation(const FMoverTickStartData& StartState, const FMoverTimeStep& TimeStep, FProposedMove& OutProposedMove) const
@@ -95,18 +104,38 @@ void UChaosFallingMode::GenerateMove_Implementation(const FMoverTickStartData& S
 		Params.Deceleration = OverTerminalSpeedFallingDeceleration;
 	}
 
-	FFloorCheckResult LastFloorResult;
-	// limit our move input based on the floor we're on
-	const UMoverBlackboard* SimBlackboard = Simulation->GetBlackboard();
-	if (SimBlackboard && SimBlackboard->TryGet(CommonBlackboard::LastFloorResult, LastFloorResult))
+	FFloorCheckResult FloorResult;
+	FWaterCheckResult WaterResult;
+	UMoverBlackboard* SimBlackboard = Simulation->GetBlackboard_Mutable();
+	if (SimBlackboard)
 	{
-		if (LastFloorResult.HitResult.IsValidBlockingHit() && LastFloorResult.HitResult.Normal.Dot(DefaultSimInputs->UpDir) > UE::MoverUtils::VERTICAL_SLOPE_NORMAL_MAX_DOT && !LastFloorResult.IsWalkableFloor())
+		// Update the floor
+		UE::ChaosMover::Utils::FFloorSweepParams SweepParams{
+			.ResponseParams = DefaultSimInputs->CollisionResponseParams,
+			.QueryParams = DefaultSimInputs->CollisionQueryParams,
+			.Location = StartingSyncState->GetLocation_WorldSpace(),
+			.DeltaPos = StartingSyncState->GetVelocity_WorldSpace() * TimeStep.StepMs * 0.001f,
+			.UpDir = DefaultSimInputs->UpDir,
+			.World = DefaultSimInputs->World,
+			.QueryDistance = 1.2f * GetTargetHeight(),
+			.QueryRadius = FMath::Min(GetGroundQueryRadius(), FMath::Max(DefaultSimInputs->PawnCollisionRadius - 5.0f, 0.0f)),
+			.MaxWalkSlopeCosine = GetMaxWalkSlopeCosine(),
+			.TargetHeight = GetTargetHeight(),
+			.CollisionChannel = DefaultSimInputs->CollisionChannel
+		};
+
+		UE::ChaosMover::Utils::FloorSweep_Internal(SweepParams, FloorResult, WaterResult);
+
+		SimBlackboard->Set(CommonBlackboard::LastFloorResult, FloorResult);
+		SimBlackboard->Set(CommonBlackboard::LastWaterResult, WaterResult);
+
+		if (FloorResult.HitResult.IsValidBlockingHit() && FloorResult.HitResult.Normal.Dot(DefaultSimInputs->UpDir) > UE::MoverUtils::VERTICAL_SLOPE_NORMAL_MAX_DOT && !FloorResult.IsWalkableFloor())
 		{
 			// If acceleration is into the wall, limit contribution.
-			if (FVector::DotProduct(Params.MoveInput, LastFloorResult.HitResult.Normal) < 0.f)
+			if (FVector::DotProduct(Params.MoveInput, FloorResult.HitResult.Normal) < 0.f)
 			{
 				// Allow movement parallel to the wall, but not into it because that may push us up.
-				const FVector FallingHitNormal = FVector::VectorPlaneProject(LastFloorResult.HitResult.Normal, -DefaultSimInputs->UpDir).GetSafeNormal();
+				const FVector FallingHitNormal = FVector::VectorPlaneProject(FloorResult.HitResult.Normal, -DefaultSimInputs->UpDir).GetSafeNormal();
 				Params.MoveInput = FVector::VectorPlaneProject(Params.MoveInput, FallingHitNormal);
 			}
 		}
diff --git a/Engine/Plugins/Experimental/ChaosMover/Source/ChaosMover/Private/Character/Modes/ChaosWalkingMode.cpp b/Engine/Plugins/Experimental/ChaosMover/Source/ChaosMover/Private/Character/Modes/ChaosWalkingMode.cpp
--- a/Engine/Plugins/Experimental/ChaosMover/Source/ChaosMover/Private/Character/Modes/ChaosWalkingMode.cpp
+++ b/Engine/Plugins/Experimental/ChaosMover/Source/ChaosMover/Private/Character/Modes/ChaosWalkingMode.cpp
@@ -1,12 +1,23 @@
 // Copyright Epic Games, Inc. All Rights Reserved.
 
 #include "ChaosMover/Character/Modes/ChaosWalkingMode.h"
+
+#include "Chaos/Character/CharacterGroundConstraintSettings.h"
+#include "ChaosMover/ChaosMoverLog.h"
+#include "ChaosMover/ChaosMoverSimulation.h"
+#include "ChaosMover/ChaosMoverSimulationTypes.h"
+#include "ChaosMover/Character/Settings/SharedChaosCharacterMovementSettings.h"
+#include "ChaosMover/Character/Transitions/ChaosCharacterFallingCheck.h"
+#include "ChaosMover/Utilities/ChaosGroundMovementUtils.h"
+#include "ChaosMover/Utilities/ChaosMoverQueryUtils.h"
+#include "Math/UnitConversion.h"
 #include "MoveLibrary/MovementUtils.h"
+#include "MoveLibrary/BasedMovementUtils.h"
 #include "MoveLibrary/FloorQueryUtils.h"
-#include "MoveLibrary/WaterMovementUtils.h"
-#include "ChaosMover/ChaosMoverSimulationTypes.h"
 #include "MoveLibrary/GroundMovementUtils.h"
-#include "DefaultMovementSet/Settings/CommonLegacyMovementSettings.h"
+#include "MoveLibrary/MovementUtils.h"
+#include "MoveLibrary/WaterMovementUtils.h"
+#include "Chaos/ParticleHandle.h"
 
 #include UE_INLINE_GENERATED_CPP_BY_NAME(ChaosWalkingMode)
 
@@ -15,163 +26,225 @@ UChaosWalkingMode::UChaosWalkingMode(const FObjectInitializer& ObjectInitializer
 {
 	bSupportsAsync = true;
 	GameplayTags.AddTag(Mover_IsOnGround);
-}
 
-void UChaosWalkingMode::GenerateMove_Implementation(const FMoverTickStartData& StartState, const FMoverTimeStep& TimeStep, FProposedMove& OutProposedMove) const
-{
+	RadialForceLimit = 2000.0f;
+	SwingTorqueLimit = 3000.0f;
+	TwistTorqueLimit = 1500.0f;
 
+	Transitions.Add(CreateDefaultSubobject<UChaosCharacterFallingCheck>(TEXT("DefaultFallingCheck")));
 }
 
-/*
-* WIP
-* 
-*
-
-USTRUCT()
-struct CHAOSMOVER_API FGroundMovementModeData
+void UChaosWalkingMode::UpdateConstraintSettings(Chaos::FCharacterGroundConstraintSettings& ConstraintSettings) const
 {
-	UPROPERTY()
-	float TargetHeight = 95.0f;
-
-	UPROPERTY()
-	float QueryRadius = 30.0f;
-};
-
-
-* 
-void FloorCheck(
-	const FMoverDefaultSyncState& SyncState,
-	const FChaosMoverSimulationDefaultInputs& ChaosMoverInput,
-	const FProposedMove& ProposedMove,
-	float DeltaSeconds,
-	float TargetHeight,
-	FFloorCheckResult& OutFloorResult,
-	FWaterCheckResult& OutWaterResult,
-	FVector& OutDeltaPos);
-
-bool CanStepUpOnHitSurface(
-	const FFloorCheckResult& FloorResult,
-	float TargetHeight,
-	const UCommonLegacyMovementSettings* CommonLegacySettings)
+	Super::UpdateConstraintSettings(ConstraintSettings);
+	ConstraintSettings.FrictionForceLimit = FUnitConversion::Convert(FrictionForceLimit, EUnit::Newtons, EUnit::KilogramCentimetersPerSecondSquared);
+	ConstraintSettings.DampingFactor = GroundDamping;
+	ConstraintSettings.MotionTargetMassBias = FractionalGroundReaction;
+	ConstraintSettings.RadialForceMotionTargetScaling = FractionalRadialForceLimitScaling;
+}
+
+void UChaosWalkingMode::GenerateMove_Implementation(const FMoverTickStartData& StartState, const FMoverTimeStep& TimeStep, FProposedMove& OutProposedMove) const
 {
-	const float StepHeight = TargetHeight - FloorResult.FloorDist;
+	const FChaosMoverSimulationDefaultInputs* DefaultSimInputs = StartState.AuxState.AuxStateCollection.FindDataByType<FChaosMoverSimulationDefaultInputs>();
+	const FCharacterDefaultInputs* CharacterInputs = StartState.InputCmd.InputCollection.FindDataByType<FCharacterDefaultInputs>();
+	if (!DefaultSimInputs || !CharacterInputs)
+	{
+		UE_LOG(LogChaosMover, Warning, TEXT("ChaosFallingMode requires FChaosMoverSimulationDefaultInputs and FCharacterDefaultInputs"));
+		return;
+	}
 
-	bool bWalkable = StepHeight <= CommonLegacySettings->MaxStepHeight;
-	constexpr float MinStepHeight = 2.0f;
-	const bool SteppingUp = StepHeight > MinStepHeight;
-	if (bWalkable && SteppingUp)
+	if (!Simulation)
 	{
-		bWalkable = UGroundMovementUtils::CanStepUpOnHitSurface(FloorResult.HitResult);
+		UE_LOG(LogChaosMover, Warning, TEXT("No Simulation set on ChaosFallingMode"));
+		return;
 	}
 
-	return bWalkable;
-}
-*/
+	const FMoverDefaultSyncState* StartingSyncState = StartState.SyncState.SyncStateCollection.FindDataByType<FMoverDefaultSyncState>();
+	check(StartingSyncState);
 
-void UChaosWalkingMode::SimulationTick_Implementation(const FSimulationTickParams& Params, FMoverTickEndData& OutputState)
-{
-	/*
-	* WIP
-	* 
-	const FMoverTickStartData& StartState = Params.StartState;
-	FProposedMove ProposedMove = Params.ProposedMove;
+	const float DeltaSeconds = TimeStep.StepMs * 0.001f;
+	const FVector UpDirection = DefaultSimInputs->UpDir;
+
+	// Try to use the floor as the basis for the intended move direction (i.e. try to walk along slopes, rather than into them)
+	UMoverBlackboard* SimBlackboard = Simulation->GetBlackboard_Mutable();
+	FFloorCheckResult LastFloorResult;
+	FVector MovementNormal;
+	FVector LocalGroundVelocity = FVector::ZeroVector;
+	if (SimBlackboard && SimBlackboard->TryGet(CommonBlackboard::LastFloorResult, LastFloorResult) && LastFloorResult.IsWalkableFloor())
+	{
+		MovementNormal = LastFloorResult.HitResult.ImpactNormal;
+		LocalGroundVelocity = UChaosGroundMovementUtils::ComputeLocalGroundVelocity_Internal(StartingSyncState->GetLocation_WorldSpace(), LastFloorResult);
+	}
+	else
+	{
+		MovementNormal = UpDirection;
+	}
 
-	const FChaosMoverSimulationDefaultInputs* ChaosMoverInfo = StartState.AuxState.AuxStateCollection.FindDataByType<FChaosMoverSimulationDefaultInputs>();
-	const FVector UpDir = ChaosMoverInfo ? UMovementUtils::DeduceUpDirectionFromGravity(ChaosMoverInfo->Gravity) : FVector::UpVector;
-	const FVector GravityAcceleration = ChaosMoverInfo ? ChaosMoverInfo->Gravity : -980 * UpDir;
+	FRotator IntendedOrientation_WorldSpace;
+	// If there's no intent from input to change orientation, use the current orientation
+	if (!CharacterInputs || CharacterInputs->OrientationIntent.IsNearlyZero())
+	{
+		IntendedOrientation_WorldSpace = StartingSyncState->GetOrientation_WorldSpace();
+	}
+	else
+	{
+		IntendedOrientation_WorldSpace = CharacterInputs->GetOrientationIntentDir_WorldSpace().ToOrientationRotator();
+	}
 
-	const FMoverDefaultSyncState* StartingSyncState = StartState.SyncState.SyncStateCollection.FindDataByType<FMoverDefaultSyncState>();
-	check(StartingSyncState);
+	const FQuat WorldToGravityTransform = FQuat::FindBetweenNormals(FVector::UpVector, DefaultSimInputs->UpDir);
+	IntendedOrientation_WorldSpace = UMovementUtils::ApplyGravityToOrientationIntent(IntendedOrientation_WorldSpace, WorldToGravityTransform, bShouldCharacterRemainUpright);
 
-	FMoverDefaultSyncState& OutputSyncState = OutputState.SyncState.SyncStateCollection.FindOrAddMutableDataByType<FMoverDefaultSyncState>();
+	FGroundMoveParams Params;
 
-	const float DeltaSeconds = Params.TimeStep.StepMs * 0.001f;
+	if (CharacterInputs)
+	{
+		Params.MoveInputType = CharacterInputs->GetMoveInputType();
 
-	UMoverBlackboard* SimBlackboard = Params.SimBlackboard;
-	if (!SimBlackboard)
+		const bool bMaintainInputMagnitude = true;
+		Params.MoveInput = CharacterInputs->GetMoveInput_WorldSpace();
+		//Params.MoveInput = UPlanarConstraintUtils::ConstrainDirectionToPlane(MoverComp->GetPlanarConstraint(), CharacterInputs->GetMoveInput_WorldSpace(), bMaintainInputMagnitude);
+	}
+	else
 	{
-		OutputSyncState = *StartingSyncState;
-		return;
+		Params.MoveInputType = EMoveInputType::None;
+		Params.MoveInput = FVector::ZeroVector;
 	}
 
-	// Store the previous ground normal that was used to compute the proposed move
-	FFloorCheckResult PrevFloorResult;
-	FVector PrevGroundNormal = UpDir; 
-	if (SimBlackboard->TryGet(CommonBlackboard::LastFloorResult, PrevFloorResult))
+	const FVector CurrentRelativeVelocity = StartingSyncState->GetVelocity_WorldSpace() - LocalGroundVelocity;
+
+	Params.OrientationIntent = IntendedOrientation_WorldSpace;
+	Params.PriorVelocity = FVector::VectorPlaneProject(CurrentRelativeVelocity, MovementNormal);
+	Params.PriorOrientation = StartingSyncState->GetOrientation_WorldSpace();
+	Params.GroundNormal = MovementNormal;
+	Params.TurningRate = SharedSettings->TurningRate;
+	Params.TurningBoost = SharedSettings->TurningBoost;
+	Params.MaxSpeed = SharedSettings->MaxSpeed;
+	Params.Acceleration = SharedSettings->Acceleration;
+	Params.Deceleration = SharedSettings->Deceleration;
+	Params.DeltaSeconds = DeltaSeconds;
+	Params.WorldToGravityQuat = WorldToGravityTransform;
+	Params.UpDirection = UpDirection;
+	Params.bUseAccelerationForVelocityMove = SharedSettings->bUseAccelerationForVelocityMove;
+
+	if (Params.MoveInput.SizeSquared() > 0.f && !UMovementUtils::IsExceedingMaxSpeed(Params.PriorVelocity, SharedSettings->MaxSpeed))
+	{
+		Params.Friction = SharedSettings->GroundFriction;
+	}
+	else
 	{
-		PrevGroundNormal = PrevFloorResult.HitResult.ImpactNormal;
+		Params.Friction = SharedSettings->bUseSeparateBrakingFriction ? SharedSettings->BrakingFriction : SharedSettings->GroundFriction;
+		Params.Friction *= SharedSettings->BrakingFrictionFactor;
 	}
 
-	// Floor query
-	FFloorCheckResult FloorResult;
-	FWaterCheckResult WaterResult;
-	FVector OutDeltaPos;
+	OutProposedMove = UGroundMovementUtils::ComputeControlledGroundMove(Params);
 
-	FloorCheck(*StartingSyncState, *ChaosMoverInfo, ProposedMove, DeltaSeconds, TargetHeight, FloorResult, WaterResult, OutDeltaPos);
+	// Add back in ground velocity to make the proposed move velocity in world space
+	OutProposedMove.LinearVelocity += LocalGroundVelocity;
 
-	ProposedMove.LinearVelocity = OutDeltaPos / DeltaSeconds;
+	if (SimBlackboard)
+	{
+		// Update the floor result and check the proposed move to prevent movement onto unwalkable surfaces
+		FVector OutDeltaPos = FVector::ZeroVector;
+		FFloorCheckResult FloorResult;
+		FWaterCheckResult WaterResult;
+		GetFloorAndCheckMovement(*StartingSyncState, OutProposedMove, *DefaultSimInputs, DeltaSeconds, FloorResult, WaterResult, OutDeltaPos);
 
-	// The base movement mode does not apply gravity in walking mode so apply here
-	// This is so that the gravity in this mode will be consistent with the gravity
-	// set on the mover, not the default physics gravity
-	FVector TargetVelocity = StartingSyncState->GetVelocity_WorldSpace() + UMovementUtils::ComputeVelocityFromGravity(GravityAcceleration, DeltaSeconds);
+		OutProposedMove.LinearVelocity = OutDeltaPos / DeltaSeconds;
 
-	SimBlackboard->Set(CommonBlackboard::LastFloorResult, FloorResult);
-	SimBlackboard->Set(CommonBlackboard::LastWaterResult, WaterResult);
+		SimBlackboard->Set(CommonBlackboard::LastFloorResult, FloorResult);
+		SimBlackboard->Set(CommonBlackboard::LastWaterResult, WaterResult);
+	};
+}
 
-	const bool bStartSwimming = WaterResult.WaterSplineData.ImmersionDepth > CommonLegacySettings->SwimmingStartImmersionDepth;
+void UChaosWalkingMode::SimulationTick_Implementation(const FSimulationTickParams& Params, FMoverTickEndData& OutputState)
+{
+	const FChaosMoverSimulationDefaultInputs* DefaultSimInputs = Params.StartState.AuxState.AuxStateCollection.FindDataByType<FChaosMoverSimulationDefaultInputs>();
+	const FCharacterDefaultInputs* CharacterInputs = Params.StartState.InputCmd.InputCollection.FindDataByType<FCharacterDefaultInputs>();
+	if (!DefaultSimInputs || !CharacterInputs)
+	{
+		UE_LOG(LogChaosMover, Warning, TEXT("ChaosFallingMode requires FChaosMoverSimulationDefaultInputs and FCharacterDefaultInputs"));
+		return;
+	}
 
-	if (WaterResult.IsSwimmableVolume() && bStartSwimming)
+	if (!Simulation)
 	{
-		SwitchToState(CommonLegacySettings->SwimmingMovementModeName, Params, OutputState);
+		UE_LOG(LogChaosMover, Warning, TEXT("No Simulation set on ChaosFallingMode"));
+		return;
 	}
-	else if (FloorResult.IsWalkableFloor())
+
+	FProposedMove ProposedMove = Params.ProposedMove;
+	const FMoverDefaultSyncState* StartingSyncState = Params.StartState.SyncState.SyncStateCollection.FindDataByType<FMoverDefaultSyncState>();
+	check(StartingSyncState);
+
+	FMoverDefaultSyncState& OutputSyncState = OutputState.SyncState.SyncStateCollection.FindOrAddMutableDataByType<FMoverDefaultSyncState>();
+	OutputSyncState = *StartingSyncState;
+
+	const float DeltaSeconds = Params.TimeStep.StepMs * 0.001f;
+	const FVector UpDirection = DefaultSimInputs->UpDir;
+
+	FVector StartingGroundNormal = UpDirection;
+	const UMoverBlackboard* SimBlackboard = Simulation->GetBlackboard();
+	FFloorCheckResult FloorResult;
+	if (SimBlackboard)
 	{
-		const FVector StartGroundVelocity = UPhysicsMovementUtils::ComputeGroundVelocityFromHitResult(StartingSyncState->GetLocation_WorldSpace(), FloorResult.HitResult, DeltaSeconds);
-		const FVector ProjectedGroundVelocity = UPhysicsMovementUtils::ComputeIntegratedGroundVelocityFromHitResult(StartingSyncState->GetLocation_WorldSpace(), FloorResult.HitResult, DeltaSeconds);
-		const bool bIsGroundMoving = ProjectedGroundVelocity.SizeSquared() > UE_KINDA_SMALL_NUMBER;
+		SimBlackboard->TryGet(CommonBlackboard::LastFloorResult, FloorResult);
+		StartingGroundNormal = FloorResult.HitResult.ImpactNormal;
+	};
+
+	if (FloorResult.IsWalkableFloor())
+	{
+		const float InitialHeightAboveFloor = FloorResult.FloorDist - GetTargetHeight();
 
-		TargetVelocity = StartingSyncState->GetVelocity_WorldSpace();
+		// Put the target position on the floor at the target height
 		FVector TargetPosition = StartingSyncState->GetLocation_WorldSpace();
-		if (FloorResult.FloorDist > TargetHeight + UE_KINDA_SMALL_NUMBER)
+		if (FloorResult.FloorDist > GetTargetHeight() + UE_KINDA_SMALL_NUMBER)
 		{
-			TargetPosition += UpDir * (TargetHeight - FloorResult.FloorDist);
+			TargetPosition -= UpDirection * InitialHeightAboveFloor;
 		}
 
-		constexpr float ParallelCosThreshold = 0.999f;
-		const bool bVerticalVelocity = FVector::Parallel(TargetVelocity.GetSafeNormal(), UpDir, ParallelCosThreshold);
-		const bool bUseProposedMove = !(bHandleVerticalLandingSeparately && bVerticalVelocity) || ProposedMove.bHasDirIntent;
+		// The base movement mode does not apply gravity in walking mode so apply here.
+		// Also remove the gravity that will be applied by the physics simulation.
+		// This is so that the gravity in this mode will be consistent with the gravity
+		// set on the mover, not the default physics gravity
+		const FVector ProjectedVelocity = StartingSyncState->GetVelocity_WorldSpace() + DefaultSimInputs->Gravity * DeltaSeconds;
+		FVector TargetVelocity = ProjectedVelocity - DefaultSimInputs->PhysicsObjectGravity * DefaultSimInputs->UpDir * DeltaSeconds;
 
+		// If we have movement intent and not moving straight up/down then use the proposed move plane velocity
+		// otherwise just fall with gravity
+		constexpr float ParallelCosThreshold = 0.999f;
+		const bool bNonVerticalVelocity = !FVector::Parallel(TargetVelocity.GetSafeNormal(), UpDirection, ParallelCosThreshold);
+		const bool bUseProposedMove = bNonVerticalVelocity || ProposedMove.bHasDirIntent;
 		if (bUseProposedMove)
 		{
-			const FVector ProposedMovePlaneVelocity = ProposedMove.LinearVelocity - ProposedMove.LinearVelocity.ProjectOnToNormal(PrevGroundNormal);
+			const FVector ProposedMovePlaneVelocity = ProposedMove.LinearVelocity - ProposedMove.LinearVelocity.ProjectOnToNormal(StartingGroundNormal);
 
 			// If there is velocity intent in the normal direction then use the velocity from the proposed move. Otherwise
 			// retain the previous vertical velocity
 			FVector ProposedNormalVelocity = ProposedMove.LinearVelocity - ProposedMovePlaneVelocity;
 			if (ProposedNormalVelocity.SizeSquared() > UE_KINDA_SMALL_NUMBER)
 			{
-				TargetVelocity += ProposedNormalVelocity - TargetVelocity.ProjectOnToNormal(PrevGroundNormal);
+				TargetVelocity += ProposedNormalVelocity - TargetVelocity.ProjectOnToNormal(StartingGroundNormal);
 			}
 
 			TargetPosition += ProposedMovePlaneVelocity * DeltaSeconds;
 		}
 
-		// Check if the proposed velocity would lift off the movement surface.
-		bool bIsLiftingOffSurface = false;
-
-		float CharacterGravity = 0.0f;
-		if (const APhysicsVolume* PhysVolume = UpdatedComponent->GetPhysicsVolume())
+		FVector ProjectedGroundVelocity = UChaosGroundMovementUtils::ComputeLocalGroundVelocity_Internal(StartingSyncState->GetLocation_WorldSpace(), FloorResult);
+		const Chaos::FPBDRigidParticleHandle* GroundParticle = UChaosGroundMovementUtils::GetRigidParticleHandleFromFloorResult_Internal(FloorResult);
+		if (GroundParticle && GroundParticle->IsDynamic() && GroundParticle->GravityEnabled())
 		{
-			CharacterGravity = PhysVolume->GetGravityZ();
+			// This might not be correct if different physics objects have different gravity but is saves having to go
+			// to the component to get the gravity on the physics volume.
+			ProjectedGroundVelocity += DefaultSimInputs->PhysicsObjectGravity * DefaultSimInputs->UpDir * DeltaSeconds;
 		}
-
-		FVector ProjectedVelocity = TargetVelocity + CharacterGravity * FVector::UpVector * DeltaSeconds;
-
-		const float ProjectedRelativeVerticalVelocity = FloorResult.HitResult.ImpactNormal.Dot(ProjectedVelocity - ProjectedGroundVelocity);
+		const bool bIsGroundMoving = ProjectedGroundVelocity.SizeSquared() > UE_KINDA_SMALL_NUMBER;
+		const FVector ProjectedRelativeVelocity = ProjectedVelocity - ProjectedGroundVelocity;
+		const float ProjectedRelativeNormalVelocity = FloorResult.HitResult.ImpactNormal.Dot(ProjectedVelocity - ProjectedGroundVelocity);
+		const float ProjectedRelativeVerticalVelocity = UpDirection.Dot(ProjectedVelocity - ProjectedGroundVelocity);
 		const float VerticalVelocityLimit = 2.0f / DeltaSeconds;
 
-		if ((ProjectedRelativeVerticalVelocity > VerticalVelocityLimit) && bIsGroundMoving && (ProjectedVelocity.Dot(UpDir) > VerticalVelocityLimit))
+		bool bIsLiftingOffSurface = false;
+		if ((ProjectedRelativeNormalVelocity > VerticalVelocityLimit) && bIsGroundMoving && (ProjectedRelativeVerticalVelocity > VerticalVelocityLimit))
 		{
 			bIsLiftingOffSurface = true;
 		}
@@ -180,108 +253,75 @@ void UChaosWalkingMode::SimulationTick_Implementation(const FSimulationTickParam
 		// If stepping up make sure that the step height is less than the max step height
 		// and the new surface has CanCharacterStepUpOn set to true.
 		// If stepping down make sure the step height is less than the max step height.
-		const float StartHeightAboveGround = FloorResult.FloorDist - TargetHeight;
-		const float EndHeightAboutGround = StartHeightAboveGround + UpDir.Dot(ProjectedVelocity - ProjectedGroundVelocity) * DeltaSeconds ;
-		const bool bIsSteppingDown = StartHeightAboveGround > GPhysicsDrivenMotionDebugParams.MinStepUpDistance;
-		const bool bIsWithinReach = EndHeightAboutGround <= CommonLegacySettings->MaxStepHeight;
-
-		// If the character is unsupported allow some grace period before falling
-		bool bIsSupported = bIsWithinReach && !bIsLiftingOffSurface;
-		float TimeSinceSupported = MaxUnsupportedTimeBeforeFalling;
-		SimBlackboard->TryGet(CommonBlackboard::TimeSinceSupported, TimeSinceSupported);
-		if (bIsSupported)
-		{
-			SimBlackboard->Set(CommonBlackboard::TimeSinceSupported, 0.0f);
-		}
-		else if (!bIsLiftingOffSurface)
-		{
-			// Falling
-			TimeSinceSupported += DeltaSeconds;
-			SimBlackboard->Set(CommonBlackboard::TimeSinceSupported, TimeSinceSupported);
-			bIsSupported = TimeSinceSupported < MaxUnsupportedTimeBeforeFalling;
-		}
-		else
-		{
-			// Moving up relative to ground
-			SimBlackboard->Set(CommonBlackboard::TimeSinceSupported, MaxUnsupportedTimeBeforeFalling);
-		}
-
-		// Apply vertical velocity to target if stepping down
-		const bool bNeedsVerticalVelocityToTarget = bIsSupported && bIsSteppingDown && (EndHeightAboutGround > 0.0f) && !bIsLiftingOffSurface;
+		const float EndHeightAboveFloor = InitialHeightAboveFloor + ProjectedRelativeVerticalVelocity * DeltaSeconds;
+		const bool bIsSteppingDown = InitialHeightAboveFloor > UE_KINDA_SMALL_NUMBER;
+		const bool bIsWithinReach = EndHeightAboveFloor <= SharedSettings->MaxStepHeight;
+		const bool bIsSupported = bIsWithinReach && !bIsLiftingOffSurface;
+		const bool bNeedsVerticalVelocityToTarget = bIsSupported && bIsSteppingDown && (EndHeightAboveFloor > 0.0f) && !bIsLiftingOffSurface;
 		if (bNeedsVerticalVelocityToTarget)
 		{
-			TargetVelocity -= FractionalDownwardVelocityToTarget * (EndHeightAboutGround / DeltaSeconds) * UpDir;
+			TargetVelocity -= FractionalDownwardVelocityToTarget * (EndHeightAboveFloor / DeltaSeconds) * UpDirection;
 		}
 
 		// Target orientation
 		// This is always applied regardless of whether the character is supported
 		FRotator TargetOrientation = StartingSyncState->GetOrientation_WorldSpace();
-		if (!UMovementUtils::IsAngularVelocityZero(ProposedMove.AngularVelocity)) 
+		if (!UMovementUtils::IsAngularVelocityZero(ProposedMove.AngularVelocity))
 		{
 			TargetOrientation += (ProposedMove.AngularVelocity * DeltaSeconds);
 		}
 
-		if (bIsSupported)
-		{
-			OutputState.MovementEndState.RemainingMs = 0.0f;
-			OutputSyncState.MoveDirectionIntent = ProposedMove.bHasDirIntent ? ProposedMove.DirectionIntent : FVector::ZeroVector;
-			OutputSyncState.SetTransforms_WorldSpace(
-				TargetPosition,
-				TargetOrientation,
-				TargetVelocity,
-				nullptr);
-		}
-		else
-		{
-			// Blocking hit but not supported
-			SwitchToState(CommonLegacySettings->AirMovementModeName, Params, OutputState);
-		}
+		OutputSyncState.SetTransforms_WorldSpace(TargetPosition, TargetOrientation, TargetVelocity);
 	}
-	else
+	
+	OutputState.MovementEndState.RemainingMs = 0.0f;
+	OutputSyncState.MoveDirectionIntent = ProposedMove.bHasDirIntent ? ProposedMove.DirectionIntent : FVector::ZeroVector;
+}
+
+bool UChaosWalkingMode::CanStepUpOnHitSurface(const FFloorCheckResult& FloorResult) const
+{
+	const float StepHeight = GetTargetHeight() - FloorResult.FloorDist;
+
+	bool bWalkable = StepHeight <= SharedSettings->MaxStepHeight;
+	constexpr float MinStepHeight = 2.0f;
+	const bool SteppingUp = StepHeight > MinStepHeight;
+	if (bWalkable && SteppingUp)
 	{
-		// No water or floor not found
-		float TimeSinceSupported = MaxUnsupportedTimeBeforeFalling;
-		SimBlackboard->TryGet(CommonBlackboard::TimeSinceSupported, TimeSinceSupported);
-		TimeSinceSupported += DeltaSeconds;
-		SimBlackboard->Set(CommonBlackboard::TimeSinceSupported, TimeSinceSupported);
-		if (TimeSinceSupported >= MaxUnsupportedTimeBeforeFalling)
-		{
-			SwitchToState(CommonLegacySettings->AirMovementModeName, Params, OutputState);
-		}
+		bWalkable = UGroundMovementUtils::CanStepUpOnHitSurface(FloorResult.HitResult);
 	}
-	*/
-}
 
+	return bWalkable;
+}
 
-/*
-*
-* WIP
-void FloorCheck(const FMoverDefaultSyncState& SyncState,
-	const FChaosMoverSimulationDefaultInputs& ChaosMoverInput,
-	const FProposedMove& ProposedMove,
-	float DeltaSeconds,
-	float TargetHeight,
-	FFloorCheckResult& OutFloorResult,
-	FWaterCheckResult& OutWaterResult,
-	FVector& OutDeltaPos)
+void UChaosWalkingMode::GetFloorAndCheckMovement(const FMoverDefaultSyncState& SyncState, const FProposedMove& ProposedMove, const FChaosMoverSimulationDefaultInputs& DefaultSimInputs, float DeltaSeconds, FFloorCheckResult& OutFloorResult, FWaterCheckResult& OutWaterResult, FVector& OutDeltaPos) const
 {
-	const FVector& UpDir = ChaosMoverInput.UpDir;
-
 	FVector DeltaPos = ProposedMove.LinearVelocity * DeltaSeconds;
 	OutDeltaPos = DeltaPos;
 
-	const float FloorSweepDistance = TargetHeight + CommonLegacySettings->MaxStepHeight;
-
-	UPhysicsMovementUtils::FloorSweep_Internal(SyncState.GetLocation_WorldSpace(), DeltaPos, UpdatedPrimitive, UpDir,
-		QueryRadius, FloorSweepDistance, CommonLegacySettings->MaxWalkSlopeCosine, TargetHeight, OutFloorResult, OutWaterResult);
+	UE::ChaosMover::Utils::FFloorSweepParams SweepParams{
+		.ResponseParams = DefaultSimInputs.CollisionResponseParams,
+		.QueryParams = DefaultSimInputs.CollisionQueryParams,
+		.Location = SyncState.GetLocation_WorldSpace(),
+		.DeltaPos = DeltaPos,
+		.UpDir = DefaultSimInputs.UpDir,
+		.World = DefaultSimInputs.World,
+		.QueryDistance = GetTargetHeight() + SharedSettings->MaxStepHeight,
+		.QueryRadius = FMath::Min(GetGroundQueryRadius(), FMath::Max(DefaultSimInputs.PawnCollisionRadius - 5.0f, 0.0f)),
+		.MaxWalkSlopeCosine = GetMaxWalkSlopeCosine(),
+		.TargetHeight = GetTargetHeight(),
+		.CollisionChannel = DefaultSimInputs.CollisionChannel
+	};
+
+	// First, try a sweep at the end position
+	UE::ChaosMover::Utils::FloorSweep_Internal(SweepParams, OutFloorResult, OutWaterResult);
 
 	if (!OutFloorResult.bBlockingHit)
 	{
-		// Floor not found
+		// No result at the end position. Fall back on the current floor result
 		return;
 	}
 
-	bool bWalkableFloor = OutFloorResult.bWalkableFloor && CanStepUpOnHitSurface(OutFloorResult, TargetHeight, );
+	bool bWalkableFloor = OutFloorResult.bWalkableFloor && CanStepUpOnHitSurface(OutFloorResult);
 	if (bWalkableFloor)
 	{
 		// Walkable floor found
@@ -289,16 +329,16 @@ void FloorCheck(const FMoverDefaultSyncState& SyncState,
 	}
 
 	// Hit something but not walkable. Try a new query to find a walkable surface
-	const float StepBlockedHeight = TargetHeight - ChaosMoverInput.BoundingCylinderfHeight + ChaosMoverInput.BoundingCylinderRadius;
-	const float StepHeight = TargetHeight - OutFloorResult.FloorDist;
+	const float StepBlockedHeight = GetTargetHeight() - DefaultSimInputs.PawnCollisionHalfHeight + DefaultSimInputs.PawnCollisionRadius;
+	const float StepHeight = GetTargetHeight() - OutFloorResult.FloorDist;
 
 	if (StepHeight > StepBlockedHeight)
 	{
 		// Collision should prevent movement. Just try to find ground at start of movement
-		const float ShrinkMultiplier = 0.75f;
-		UPhysicsMovementUtils::FloorSweep_Internal(SyncState.GetLocation_WorldSpace(), FVector::ZeroVector, UpdatedPrimitive, UpDir,
-			ShrinkMultiplier * QueryRadius, FloorSweepDistance, CommonLegacySettings->MaxWalkSlopeCosine, TargetHeight, OutFloorResult, OutWaterResult);
+		SweepParams.QueryRadius = 0.75f * GetGroundQueryRadius();
+		SweepParams.DeltaPos = FVector::ZeroVector;
 
+		UE::ChaosMover::Utils::FloorSweep_Internal(SweepParams, OutFloorResult, OutWaterResult);
 		OutFloorResult.bWalkableFloor = OutFloorResult.bWalkableFloor && CanStepUpOnHitSurface(OutFloorResult);
 		return;
 	}
@@ -311,10 +351,7 @@ void FloorCheck(const FMoverDefaultSyncState& SyncState,
 	}
 
 	// Try to limit the movement to remain on a walkable surface
-	FVector NewDeltaPos = DeltaPos;
-	float NewQueryRadius = QueryRadius;
-
-	FVector HorizSurfaceDir = FVector::VectorPlaneProject(OutFloorResult.HitResult.ImpactNormal, UpDir);
+	FVector HorizSurfaceDir = FVector::VectorPlaneProject(OutFloorResult.HitResult.ImpactNormal, DefaultSimInputs.UpDir);
 	float HorizSurfaceDirSizeSq = HorizSurfaceDir.SizeSquared();
 	bool bFoundOutwardDir = false;
 	if (HorizSurfaceDirSizeSq > UE_SMALL_NUMBER)
@@ -325,7 +362,7 @@ void FloorCheck(const FMoverDefaultSyncState& SyncState,
 	else
 	{
 		// Flat unwalkable surface. Try and get the horizontal direction from the normal instead
-		HorizSurfaceDir = FVector::VectorPlaneProject(OutFloorResult.HitResult.Normal, UpDir);
+		HorizSurfaceDir = FVector::VectorPlaneProject(OutFloorResult.HitResult.Normal, DefaultSimInputs.UpDir);
 		HorizSurfaceDirSizeSq = HorizSurfaceDir.SizeSquared();
 
 		if (HorizSurfaceDirSizeSq > UE_SMALL_NUMBER)
@@ -337,43 +374,31 @@ void FloorCheck(const FMoverDefaultSyncState& SyncState,
 
 	if (bFoundOutwardDir)
 	{
-		// If we're moving away try a ray query at the end of the motion
 		const float DP = DeltaPos.Dot(HorizSurfaceDir);
+		FVector NewDeltaPos = DeltaPos;
 		if (DP > 0.0f)
 		{
-			NewQueryRadius = 0.0f;
+			// If we're moving away try a ray query at the end of the motion
+			SweepParams.QueryRadius = 0.0f;
 		}
 		else
 		{
-			NewQueryRadius = 0.25f * QueryRadius;
+			// Otherwise, try to find a walkable floor along the surface
+			SweepParams.QueryRadius = 0.25f * GetGroundQueryRadius();
 			NewDeltaPos = DeltaPos - DP * HorizSurfaceDir;
 		}
+		SweepParams.DeltaPos = NewDeltaPos;
 
-		UPhysicsMovementUtils::FloorSweep_Internal(SyncState.GetLocation_WorldSpace(), NewDeltaPos, UpdatedPrimitive, UpDir,
-			NewQueryRadius, FloorSweepDistance, CommonLegacySettings->MaxWalkSlopeCosine, TargetHeight, OutFloorResult, OutWaterResult);
-
+		UE::ChaosMover::Utils::FloorSweep_Internal(SweepParams, OutFloorResult, OutWaterResult);
 		OutFloorResult.bWalkableFloor = OutFloorResult.bWalkableFloor && CanStepUpOnHitSurface(OutFloorResult);
+
 		if (OutFloorResult.bWalkableFloor)
 		{
 			OutDeltaPos = NewDeltaPos;
 		}
-		else
-		{
-			OutFloorResult.bWalkableFloor = false;
-		}
 	}
 	else
 	{
-		// Try a query at the start of the movement to find a walkable surface and prevent movement
-
-		NewDeltaPos = FVector::ZeroVector;
-		NewQueryRadius = 0.25f * QueryRadius;
-
-		UPhysicsMovementUtils::FloorSweep_Internal(SyncState.GetLocation_WorldSpace(), NewDeltaPos, UpdatedPrimitive, UpDir,
-			NewQueryRadius, FloorSweepDistance, CommonLegacySettings->MaxWalkSlopeCosine, TargetHeight, OutFloorResult, OutWaterResult);
-
-		OutFloorResult.bWalkableFloor = OutFloorResult.bWalkableFloor && CanStepUpOnHitSurface(OutFloorResult);
-		OutDeltaPos = NewDeltaPos;
+		OutDeltaPos = FVector::ZeroVector;
 	}
 }
-*/
diff --git a/Engine/Plugins/Experimental/ChaosMover/Source/ChaosMover/Private/Character/Transitions/ChaosCharacterFallingCheck.cpp b/Engine/Plugins/Experimental/ChaosMover/Source/ChaosMover/Private/Character/Transitions/ChaosCharacterFallingCheck.cpp
new file mode 100644
--- /dev/null
+++ b/Engine/Plugins/Experimental/ChaosMover/Source/ChaosMover/Private/Character/Transitions/ChaosCharacterFallingCheck.cpp
@@ -0,0 +1,132 @@
+// Copyright Epic Games, Inc. All Rights Reserved.
+
+#include "ChaosMover/Character/Transitions/ChaosCharacterFallingCheck.h"
+
+#include "Chaos/ParticleHandle.h"
+#include "ChaosMover/ChaosMoverLog.h"
+#include "ChaosMover/ChaosMoverSimulation.h"
+#include "ChaosMover/ChaosMoverSimulationTypes.h"
+#include "ChaosMover/Character/Settings/SharedChaosCharacterMovementSettings.h"
+#include "ChaosMover/Utilities/ChaosGroundMovementUtils.h"
+#include "MoveLibrary/FloorQueryUtils.h"
+#include "MoverComponent.h"
+#include "MoverSimulationTypes.h"
+
+#include UE_INLINE_GENERATED_CPP_BY_NAME(ChaosCharacterFallingCheck)
+
+UChaosCharacterFallingCheck::UChaosCharacterFallingCheck(const FObjectInitializer& ObjectInitializer)
+	: Super(ObjectInitializer)
+{
+	bSupportsAsync = true;
+
+	TransitionToFallingMode = DefaultModeNames::Falling;
+	SharedSettingsClasses.Add(USharedChaosCharacterMovementSettings::StaticClass());
+}
+
+void UChaosCharacterFallingCheck::OnRegistered()
+{
+	Super::OnRegistered();
+
+	SharedSettings = GetMoverComponent()->FindSharedSettings<USharedChaosCharacterMovementSettings>();
+	ensureMsgf(SharedSettings, TEXT("Failed to find instance of USharedChaosCharacterMovementSettings on %s. Movement may not function properly."), *GetPathNameSafe(this));
+}
+
+void UChaosCharacterFallingCheck::OnUnregistered()
+{
+	SharedSettings = nullptr;
+
+	Super::OnUnregistered();
+}
+
+FTransitionEvalResult UChaosCharacterFallingCheck::Evaluate_Implementation(const FSimulationTickParams& Params) const
+{
+	FTransitionEvalResult EvalResult;
+
+	if (TransitionToFallingMode.IsNone())
+	{
+		return EvalResult;
+	}
+	
+	check(Simulation);
+
+	const FMoverTickStartData& StartState = Params.StartState;
+
+	if (const FChaosMoverSimulationDefaultInputs* SimInputs = StartState.AuxState.AuxStateCollection.FindDataByType<FChaosMoverSimulationDefaultInputs>())
+	{
+		UMoverBlackboard* Blackboard = Simulation->GetBlackboard_Mutable();
+		check(Blackboard);
+
+		FFloorCheckResult FloorResult;
+		if (!Blackboard->TryGet(CommonBlackboard::LastFloorResult, FloorResult))
+		{
+			return EvalResult;
+		}
+
+		const FCharacterDefaultInputs* CharacterInputs = StartState.InputCmd.InputCollection.FindDataByType<FCharacterDefaultInputs>();
+		const FMoverDefaultSyncState* StartingSyncState = StartState.SyncState.SyncStateCollection.FindDataByType<FMoverDefaultSyncState>();
+		if (!CharacterInputs || !StartingSyncState)
+		{
+			return EvalResult;
+		}
+
+		if (const IChaosCharacterMovementModeInterface* Mode = Cast<IChaosCharacterMovementModeInterface>(Simulation->FindMovementModeByName(Params.StartState.SyncState.MovementMode)))
+		{
+			const float DeltaSeconds = Params.TimeStep.StepMs * 0.001f;
+			const FVector LinearVelocity = StartingSyncState->GetVelocity_WorldSpace();
+
+			const FVector ProjectedVelocity = StartingSyncState->GetVelocity_WorldSpace() + SimInputs->Gravity * DeltaSeconds;
+			FVector ProjectedGroundVelocity = UChaosGroundMovementUtils::ComputeLocalGroundVelocity_Internal(StartingSyncState->GetLocation_WorldSpace(), FloorResult);
+			const Chaos::FPBDRigidParticleHandle* GroundParticle = UChaosGroundMovementUtils::GetRigidParticleHandleFromFloorResult_Internal(FloorResult);
+			if (GroundParticle && GroundParticle->IsDynamic() && GroundParticle->GravityEnabled())
+			{
+				// This might not be correct if different physics objects have different gravity but is saves having to go
+				// to the component to get the gravity on the physics volume.
+				ProjectedGroundVelocity += SimInputs->PhysicsObjectGravity * SimInputs->UpDir * DeltaSeconds;
+			}
+			const bool bIsGroundMoving = ProjectedGroundVelocity.SizeSquared() > UE_KINDA_SMALL_NUMBER;
+			const FVector ProjectedRelativeVelocity = ProjectedVelocity - ProjectedGroundVelocity;
+			const float ProjectedRelativeNormalVelocity = FloorResult.HitResult.ImpactNormal.Dot(ProjectedVelocity - ProjectedGroundVelocity);
+			const float ProjectedRelativeVerticalVelocity = SimInputs->UpDir.Dot(ProjectedVelocity - ProjectedGroundVelocity);
+			const float VerticalVelocityLimit = 2.0f / DeltaSeconds;
+			const bool bJumping = CharacterInputs && CharacterInputs->bIsJumpPressed;
+			bool bIsLiftingOffSurface = false;
+			if (bJumping || ((ProjectedRelativeNormalVelocity > VerticalVelocityLimit) && bIsGroundMoving && (ProjectedRelativeVerticalVelocity > VerticalVelocityLimit)))
+			{
+				bIsLiftingOffSurface = true;
+			}
+
+			const bool bIsWithinReach = FloorResult.FloorDist - Mode->GetTargetHeight() <= SharedSettings->MaxStepHeight;
+			bool bIsSupported = bIsWithinReach && !bIsLiftingOffSurface;
+			float TimeSinceSupported = MaxUnsupportedTimeBeforeFalling;
+			Blackboard->TryGet(CommonBlackboard::TimeSinceSupported, TimeSinceSupported);
+
+			if (bIsSupported)
+			{
+				Blackboard->Set(CommonBlackboard::TimeSinceSupported, 0.0f);
+			}
+			else if (!bIsLiftingOffSurface)
+			{
+				// Falling
+				TimeSinceSupported += DeltaSeconds;
+				Blackboard->Set(CommonBlackboard::TimeSinceSupported, TimeSinceSupported);
+				bIsSupported = TimeSinceSupported < MaxUnsupportedTimeBeforeFalling;
+			}
+			else
+			{
+				// Moving up relative to ground
+				Blackboard->Set(CommonBlackboard::TimeSinceSupported, MaxUnsupportedTimeBeforeFalling);
+			}
+
+			if (!bIsSupported)
+			{
+				EvalResult.NextMode = TransitionToFallingMode;
+			}
+		}
+	}
+
+	return EvalResult;
+}
+
+void UChaosCharacterFallingCheck::Trigger_Implementation(const FSimulationTickParams& Params)
+{
+}
\ No newline at end of file
diff --git a/Engine/Plugins/Experimental/ChaosMover/Source/ChaosMover/Private/Character/Transitions/ChaosCharacterLandingCheck.cpp b/Engine/Plugins/Experimental/ChaosMover/Source/ChaosMover/Private/Character/Transitions/ChaosCharacterLandingCheck.cpp
--- a/Engine/Plugins/Experimental/ChaosMover/Source/ChaosMover/Private/Character/Transitions/ChaosCharacterLandingCheck.cpp
+++ b/Engine/Plugins/Experimental/ChaosMover/Source/ChaosMover/Private/Character/Transitions/ChaosCharacterLandingCheck.cpp
@@ -36,8 +36,10 @@ FTransitionEvalResult UChaosCharacterLandingCheck::Evaluate_Implementation(const
 
 		const FCharacterDefaultInputs* CharacterInputs = StartState.InputCmd.InputCollection.FindDataByType<FCharacterDefaultInputs>();
 		const FMoverDefaultSyncState* StartingSyncState = StartState.SyncState.SyncStateCollection.FindDataByType<FMoverDefaultSyncState>();
-		check(CharacterInputs);
-		check(StartingSyncState);
+		if (!CharacterInputs || !StartingSyncState)
+		{
+			return EvalResult;
+		}
 
 		if (const IChaosCharacterMovementModeInterface* Mode = Cast<IChaosCharacterMovementModeInterface>(Simulation->FindMovementModeByName(Params.StartState.SyncState.MovementMode)))
 		{
diff --git a/Engine/Plugins/Experimental/ChaosMover/Source/ChaosMover/Public/ChaosMover/ChaosMoverSimulation.h b/Engine/Plugins/Experimental/ChaosMover/Source/ChaosMover/Public/ChaosMover/ChaosMoverSimulation.h
--- a/Engine/Plugins/Experimental/ChaosMover/Source/ChaosMover/Public/ChaosMover/ChaosMoverSimulation.h
+++ b/Engine/Plugins/Experimental/ChaosMover/Source/ChaosMover/Public/ChaosMover/ChaosMoverSimulation.h
@@ -86,7 +86,6 @@ protected:
 	virtual void OnPostSimulationTick(const FMoverTimeStep& TimeStep, UE::ChaosMover::FSimulationOutputData& OutputData);
 	virtual void OnModifyContacts(Chaos::FCollisionContactModifier& Modifier) const;
 
-	virtual void PreSimulationTickCharacter(const IChaosCharacterMovementModeInterface& CharacterMode, const FMoverTimeStep& TimeStep, const UE::ChaosMover::FSimulationInputData& InputData);
 	virtual void PostSimulationTickCharacter(const IChaosCharacterMovementModeInterface& CharacterMode, const FMoverTimeStep& TimeStep, UE::ChaosMover::FSimulationOutputData& OutputData);
 	virtual void PostSimulationTickCharacterConstraint(const IChaosCharacterConstraintMovementModeInterface& CharacterConstraintMode, const FMoverTimeStep& TimeStep, UE::ChaosMover::FSimulationOutputData& OutputData);
 
@@ -95,8 +94,8 @@ protected:
 	void DisableCharacterConstraint();
 
 	FMoverSyncState CurrentSyncState;
-	FVector DeltaPosition;
-	FRotator DeltaRotation;
+	FVector TargetDeltaPosition;
+	float TargetDeltaFacing;
 
 	UE::ChaosMover::FMoverStateMachine StateMachine;
 
diff --git a/Engine/Plugins/Experimental/ChaosMover/Source/ChaosMover/Public/ChaosMover/ChaosMoverSimulationTypes.h b/Engine/Plugins/Experimental/ChaosMover/Source/ChaosMover/Public/ChaosMover/ChaosMoverSimulationTypes.h
--- a/Engine/Plugins/Experimental/ChaosMover/Source/ChaosMover/Public/ChaosMover/ChaosMoverSimulationTypes.h
+++ b/Engine/Plugins/Experimental/ChaosMover/Source/ChaosMover/Public/ChaosMover/ChaosMoverSimulationTypes.h
@@ -49,7 +49,7 @@ struct CHAOSMOVER_API FChaosMoverSimulationDefaultInputs : public FMoverDataStru
 	/** return newly allocated copy of this FMoverDataStructBase. MUST be overridden by derived types. */
 	virtual FMoverDataStructBase* Clone() const override
 	{
-		return new FChaosMoverSimulationDefaultInputs();
+		return new FChaosMoverSimulationDefaultInputs(*this);
 	}
 
 	/** Gets the type info of this FMoverDataStructBase. MUST be overridden by derived types. */
@@ -67,9 +67,10 @@ struct CHAOSMOVER_API FChaosMoverSimulationDefaultInputs : public FMoverDataStru
 	FVector Gravity;
 	Chaos::FPhysicsObjectHandle PhysicsObject;
 	AActor* OwningActor;
-	float BoundingCylinderRadius;
-	float BoundingCylinderfHeight;
+	UWorld* World;
 	float PhysicsObjectGravity;
+	float PawnCollisionHalfHeight;
+	float PawnCollisionRadius;
 	ECollisionChannel CollisionChannel;
 };
 
diff --git a/Engine/Plugins/Experimental/ChaosMover/Source/ChaosMover/Public/ChaosMover/Character/Modes/ChaosWalkingMode.h b/Engine/Plugins/Experimental/ChaosMover/Source/ChaosMover/Public/ChaosMover/Character/Modes/ChaosWalkingMode.h
--- a/Engine/Plugins/Experimental/ChaosMover/Source/ChaosMover/Public/ChaosMover/Character/Modes/ChaosWalkingMode.h
+++ b/Engine/Plugins/Experimental/ChaosMover/Source/ChaosMover/Public/ChaosMover/Character/Modes/ChaosWalkingMode.h
@@ -6,6 +6,9 @@
 
 #include "ChaosWalkingMode.generated.h"
 
+struct FFloorCheckResult;
+struct FWaterCheckResult;
+
 /**
  * Chaos character walking mode
  */
@@ -17,6 +20,40 @@ class CHAOSMOVER_API UChaosWalkingMode : public UChaosCharacterMovementMode
 public:
 	UChaosWalkingMode(const FObjectInitializer& ObjectInitializer);
 
+	// Damping factor to control the softness of the interaction between the character and the ground
+	// Set to 0 for no damping and 1 for maximum damping
+	UPROPERTY(EditAnywhere, Category = "Constraint Settings", meta = (ClampMin = "0", UIMin = "0", ClampMax = "1", UIMax = "1"))
+	float GroundDamping = 0.0f;
+
+	// Maximum force the character can apply to hold in place while standing on an unwalkable incline
+	UPROPERTY(EditAnywhere, Category = "Constraint Settings", meta = (ClampMin = "0", UIMin = "0", ForceUnits = "Newtons"))
+	float FrictionForceLimit = 100.0f;
+
+	// Scaling applied to the radial force limit to raise the limit to always allow the character to
+	// reach the motion target/
+	// A value of 1 means that the radial force limit will be increased as needed to match the force
+	// required to achieve the motion target.
+	// A value of 0 means no scaling will be applied.
+	UPROPERTY(EditAnywhere, Category = "Constraint Settings", meta = (ClampMin = "0", UIMin = "0", ClampMax = "1", UIMax = "1"))
+	float FractionalRadialForceLimitScaling = 1.0f;
+
+	// Controls the reaction force applied to the ground in the ground plane when the character is moving
+	// A value of 1 means that the full reaction force is applied
+	// A value of 0 means the character only applies a normal force to the ground and no tangential force
+	UPROPERTY(EditAnywhere, Category = "Constraint Settings", meta = (ClampMin = "0", UIMin = "0", ClampMax = "1", UIMax = "1"))
+	float FractionalGroundReaction = 1.0f;
+
+	// Controls how much downward velocity is applied to keep the character rooted to the ground when the character
+	// is within MaxStepHeight of the ground surface.
+	UPROPERTY(EditAnywhere, Category = "Movement Settings", meta = (ClampMin = "0", UIMin = "0", ClampMax = "1", UIMax = "1"))
+	float FractionalDownwardVelocityToTarget = 1.0f;
+
+	virtual void UpdateConstraintSettings(Chaos::FCharacterGroundConstraintSettings& ConstraintSettings) const override;
+
 	virtual void GenerateMove_Implementation(const FMoverTickStartData& StartState, const FMoverTimeStep& TimeStep, FProposedMove& OutProposedMove) const override;
 	virtual void SimulationTick_Implementation(const FSimulationTickParams& Params, FMoverTickEndData& OutputState) override;
+
+protected:
+	virtual bool CanStepUpOnHitSurface(const FFloorCheckResult& FloorResult) const;
+	virtual void GetFloorAndCheckMovement(const FMoverDefaultSyncState& SyncState, const FProposedMove& ProposedMove, const FChaosMoverSimulationDefaultInputs& DefaultSimInputs, float DeltaSeconds, FFloorCheckResult& FloorResult, FWaterCheckResult& WaterResult, FVector& OutDeltaPos) const;
 };
\ No newline at end of file
diff --git a/Engine/Plugins/Experimental/ChaosMover/Source/ChaosMover/Public/ChaosMover/Character/Settings/SharedChaosCharacterMovementSettings.h b/Engine/Plugins/Experimental/ChaosMover/Source/ChaosMover/Public/ChaosMover/Character/Settings/SharedChaosCharacterMovementSettings.h
--- a/Engine/Plugins/Experimental/ChaosMover/Source/ChaosMover/Public/ChaosMover/Character/Settings/SharedChaosCharacterMovementSettings.h
+++ b/Engine/Plugins/Experimental/ChaosMover/Source/ChaosMover/Public/ChaosMover/Character/Settings/SharedChaosCharacterMovementSettings.h
@@ -53,6 +53,49 @@ public:
 	UPROPERTY(Category = "General", EditAnywhere, BlueprintReadWrite, AdvancedDisplay)
 	bool bUseAccelerationForVelocityMove = true;
 
+	/**
+	 * Setting that affects movement control. Higher values allow faster changes in direction. This can be used to simulate slippery
+	 * surfaces such as ice or oil by lowering the value (possibly based on the material the actor is standing on).
+	 */
+	UPROPERTY(BlueprintReadWrite, EditAnywhere, Category = "General|Friction", meta = (ClampMin = "0", UIMin = "0"))
+	float GroundFriction = 8.0f;
+
+	/**
+	  * If true, BrakingFriction will be used to slow the character to a stop (when there is no Acceleration).
+	  * If false, braking uses the same friction passed to CalcVelocity() (ie GroundFriction when walking), multiplied by BrakingFrictionFactor.
+	  * This setting applies to all movement modes; if only desired in certain modes, consider toggling it when movement modes change.
+	  * @see BrakingFriction
+	  */
+	UPROPERTY(Category = "General|Friction", EditDefaultsOnly, BlueprintReadWrite)
+	uint8 bUseSeparateBrakingFriction : 1;
+
+	/**
+	 * Friction (drag) coefficient applied when braking (whenever Acceleration = 0, or if character is exceeding max speed); actual value used is this multiplied by BrakingFrictionFactor.
+	 * When braking, this property allows you to control how much friction is applied when moving across the ground, applying an opposing force that scales with current velocity.
+	 * Braking is composed of friction (velocity-dependent drag) and constant deceleration.
+	 * This is the current value, used in all movement modes; if this is not desired, override it or bUseSeparateBrakingFriction when movement mode changes.
+	 * @note Only used if bUseSeparateBrakingFriction setting is true, otherwise current friction such as GroundFriction is used.
+	 * @see bUseSeparateBrakingFriction, BrakingFrictionFactor, GroundFriction, BrakingDecelerationWalking
+	 */
+	UPROPERTY(Category = "General|Friction", EditAnywhere, BlueprintReadWrite, meta = (ClampMin = "0", UIMin = "0", EditCondition = "bUseSeparateBrakingFriction"))
+	float BrakingFriction = 8.0f;
+
+	/**
+	 * Factor used to multiply actual value of friction used when braking.
+	 * This applies to any friction value that is currently used, which may depend on bUseSeparateBrakingFriction.
+	 * @note This is 2 by default for historical reasons, a value of 1 gives the true drag equation.
+	 */
+	UPROPERTY(BlueprintReadWrite, EditAnywhere, Category = "General|Friction", meta = (ClampMin = "0", UIMin = "0"))
+	float BrakingFrictionFactor = 2.0f;
+
+	/** Mover actors will be able to step up onto or over obstacles shorter than this */
+	UPROPERTY(BlueprintReadWrite, EditAnywhere, Category = "Ground Movement", meta = (ClampMin = "0", UIMin = "0", ForceUnits = "cm"))
+	float MaxStepHeight = 40.0f;
+
+	// Default movement mode to use when falling.
+	UPROPERTY(Category = "General", EditAnywhere, BlueprintReadWrite)
+	FName DefaultFallingMode = DefaultModeNames::Falling;
+
 protected:
 	/** Walkable slope angle, represented as cosine(max slope angle) for performance reasons. Ex: for max slope angle of 30 degrees, value is cosine(30 deg) = 0.866 */
 	UPROPERTY(BlueprintReadWrite, EditAnywhere, Category = "Ground Movement", BlueprintSetter=SetMaxWalkableSlopeAngle, meta = (ClampMin = "0", UIMin = "0", ClampMax = "90", UIMax = "90", ForceUnits = "degrees"))
diff --git a/Engine/Plugins/Experimental/ChaosMover/Source/ChaosMover/Public/ChaosMover/Character/Transitions/ChaosCharacterFallingCheck.h b/Engine/Plugins/Experimental/ChaosMover/Source/ChaosMover/Public/ChaosMover/Character/Transitions/ChaosCharacterFallingCheck.h
new file mode 100644
--- /dev/null
+++ b/Engine/Plugins/Experimental/ChaosMover/Source/ChaosMover/Public/ChaosMover/Character/Transitions/ChaosCharacterFallingCheck.h
@@ -0,0 +1,36 @@
+// Copyright Epic Games, Inc. All Rights Reserved.
+
+#pragma once
+
+#include "ChaosMover/ChaosMovementModeTransition.h"
+
+#include "ChaosCharacterFallingCheck.generated.h"
+
+
+UCLASS(Blueprintable, BlueprintType, EditInlineNew, DefaultToInstanced)
+class CHAOSMOVER_API UChaosCharacterFallingCheck : public UChaosMovementModeTransition
+{
+	GENERATED_BODY()
+
+public:
+	UChaosCharacterFallingCheck(const FObjectInitializer& ObjectInitializer);
+
+	virtual void OnRegistered() override;
+	virtual void OnUnregistered() override;
+
+	virtual FTransitionEvalResult Evaluate_Implementation(const FSimulationTickParams& Params) const override;
+	virtual void Trigger_Implementation(const FSimulationTickParams& Params) override;
+
+	/** Name of movement mode to transition to when landing on ground */
+	UPROPERTY(BlueprintReadWrite, EditAnywhere, Category = Parameters)
+	FName TransitionToFallingMode;
+
+	// Time limit for being unsupported before moving from a walking to a falling state.
+	// This provides some grace period when walking off of an edge during which locomotion
+	// and jumping are still possible even though the character has started falling under gravity
+	UPROPERTY(EditAnywhere, Category = Parameters, meta = (ClampMin = "0", UIMin = "0", ForceUnits = "s"))
+	float MaxUnsupportedTimeBeforeFalling = 0.06f;
+
+protected:
+	TObjectPtr<const class USharedChaosCharacterMovementSettings> SharedSettings;
+};
\ No newline at end of file
diff --git a/Engine/Plugins/Experimental/ChaosMover/Source/ChaosMover/Public/ChaosMover/Utilities/ChaosGroundMovementUtils.h b/Engine/Plugins/Experimental/ChaosMover/Source/ChaosMover/Public/ChaosMover/Utilities/ChaosGroundMovementUtils.h
--- a/Engine/Plugins/Experimental/ChaosMover/Source/ChaosMover/Public/ChaosMover/Utilities/ChaosGroundMovementUtils.h
+++ b/Engine/Plugins/Experimental/ChaosMover/Source/ChaosMover/Public/ChaosMover/Utilities/ChaosGroundMovementUtils.h
@@ -15,9 +15,9 @@ class CHAOSMOVER_API UChaosGroundMovementUtils : public UBlueprintFunctionLibrar
 	GENERATED_BODY()
 
 public:
+	/** Computes the local velocity at the supplied position of the hit object in floor result */
 	UFUNCTION(BlueprintCallable, Category = ChaosMover)
 	static FVector ComputeLocalGroundVelocity_Internal(const FVector& Position, const FFloorCheckResult& FloorResult);
 
-protected:
 	static Chaos::FPBDRigidParticleHandle* GetRigidParticleHandleFromFloorResult_Internal(const FFloorCheckResult& FloorResult);
 };
\ No newline at end of file
