Chaos Mover - Increase default landing floor distance tolerance and allow for small numerical error

diff --git a/Engine/Plugins/Experimental/ChaosMover/Source/ChaosMover/Private/Character/Transitions/ChaosCharacterLandingCheck.cpp b/Engine/Plugins/Experimental/ChaosMover/Source/ChaosMover/Private/Character/Transitions/ChaosCharacterLandingCheck.cpp
--- a/Engine/Plugins/Experimental/ChaosMover/Source/ChaosMover/Private/Character/Transitions/ChaosCharacterLandingCheck.cpp
+++ b/Engine/Plugins/Experimental/ChaosMover/Source/ChaosMover/Private/Character/Transitions/ChaosCharacterLandingCheck.cpp
@@ -85,7 +85,7 @@ FTransitionEvalResult UChaosCharacterLandingCheck::Evaluate_Implementation(const
 					const FVector LocalGroundVelocity = UChaosGroundMovementUtils::ComputeLocalGroundVelocity_Internal(StartingSyncState->GetLocation_WorldSpace(), FloorResult);
 					const float RelativeVerticalVelocity = FloorResult.HitResult.ImpactNormal.Dot(LinearVelocity - LocalGroundVelocity);
 					const float ProjectedFloorDistance = FloorDistanceWithFloorNormal + RelativeVerticalVelocity * DeltaSeconds;
-					const bool bIsFloorWithinReach = ProjectedFloorDistance < TargetHeight + FloorDistanceTolerance;
+					const bool bIsFloorWithinReach = ProjectedFloorDistance < TargetHeight + FloorDistanceTolerance + UE_KINDA_SMALL_NUMBER;
 					const bool bIsMovingUpRelativeToFloor = (RelativeVerticalVelocity > UE_KINDA_SMALL_NUMBER) || bJumping;
 
 					if (bIsFloorWithinReach && !bIsMovingUpRelativeToFloor)
diff --git a/Engine/Plugins/Experimental/ChaosMover/Source/ChaosMover/Public/ChaosMover/Character/Transitions/ChaosCharacterLandingCheck.h b/Engine/Plugins/Experimental/ChaosMover/Source/ChaosMover/Public/ChaosMover/Character/Transitions/ChaosCharacterLandingCheck.h
--- a/Engine/Plugins/Experimental/ChaosMover/Source/ChaosMover/Public/ChaosMover/Character/Transitions/ChaosCharacterLandingCheck.h
+++ b/Engine/Plugins/Experimental/ChaosMover/Source/ChaosMover/Public/ChaosMover/Character/Transitions/ChaosCharacterLandingCheck.h
@@ -24,7 +24,7 @@ public:
 
 	/** Height at which we consider the character to be on the ground */
 	UPROPERTY(BlueprintReadWrite, EditAnywhere, Category = Parameters, meta = (Units = "cm"))
-	float FloorDistanceTolerance = 0.0f;
+	float FloorDistanceTolerance = 0.5f;
 
 	/** Name of movement mode to transition to when landing on ground */
 	UPROPERTY(BlueprintReadWrite, EditAnywhere, Category = Parameters)
