Chaos Mover - Add jump transition and flying mode

diff --git a/Engine/Plugins/Experimental/ChaosMover/Source/ChaosMover/Private/ChaosMoverSimulation.cpp b/Engine/Plugins/Experimental/ChaosMover/Source/ChaosMover/Private/ChaosMoverSimulation.cpp
--- a/Engine/Plugins/Experimental/ChaosMover/Source/ChaosMover/Private/ChaosMoverSimulation.cpp
+++ b/Engine/Plugins/Experimental/ChaosMover/Source/ChaosMover/Private/ChaosMoverSimulation.cpp
@@ -144,6 +144,22 @@ void UChaosMoverSimulation::ModifyContacts(Chaos::FCollisionContactModifier& Mod
 	OnModifyContacts(Modifier);
 }
 
+void UChaosMoverSimulation::PreSimulationTickCharacter(const IChaosCharacterMovementModeInterface& CharacterMode, const FMoverTimeStep& TimeStep, const UE::ChaosMover::FSimulationInputData& InputData)
+{
+	// Add inputs if we don't have them and make sure we have a valid input
+	FCharacterDefaultInputs& CharacterDefaultInputs = InputData.InputCmd.InputCollection.FindOrAddMutableDataByType<FCharacterDefaultInputs>();
+	if (CharacterDefaultInputs.GetMoveInputType() == EMoveInputType::Invalid)
+	{
+		CharacterDefaultInputs.SetMoveInput(EMoveInputType::DirectionalIntent, FVector::ZeroVector);
+	}
+
+	if (!CharacterDefaultInputs.SuggestedMovementMode.IsNone())
+	{
+		StateMachine.QueueNextMode(CharacterDefaultInputs.SuggestedMovementMode);
+		CharacterDefaultInputs.SuggestedMovementMode = NAME_None;
+	}
+}
+
 void UChaosMoverSimulation::OnPreSimulationTick(const FMoverTimeStep& TimeStep, const UE::ChaosMover::FSimulationInputData& InputData)
 {
 	if (bInputCmdOverridden)
@@ -167,6 +183,14 @@ void UChaosMoverSimulation::OnPreSimulationTick(const FMoverTimeStep& TimeStep,
 	{
 		DefaultSyncState.SetTransforms_WorldSpace(ParticleHandle->GetX(), FRotator(ParticleHandle->GetR()), ParticleHandle->GetV());
 	}
+
+	if (TStrongObjectPtr<UBaseMovementMode> CurrentModePtr = StateMachine.GetCurrentMode().Pin())
+	{
+		if (IChaosCharacterMovementModeInterface* CharacterMode = Cast<IChaosCharacterMovementModeInterface>(CurrentModePtr.Get()))
+		{
+			PreSimulationTickCharacter(*CharacterMode, TimeStep, InputData);
+		}
+	}
 }
 
 void UChaosMoverSimulation::OnSimulationTick(const FMoverTimeStep& TimeStep, const UE::ChaosMover::FSimulationInputData& InputData, UE::ChaosMover::FSimulationOutputData& OutputData)
@@ -394,4 +418,41 @@ void UChaosMoverSimulation::DisableCharacterConstraint()
 
 		ConstraintHandle->SetEnabled(false);
 	}
-}
\ No newline at end of file
+}
+
+void UChaosMoverSimulation::K2_QueueInstantMovementEffect(const int32& EffectAsRawData)
+{
+	// This will never be called, the exec version below will be hit instead
+	checkNoEntry();
+}
+
+DEFINE_FUNCTION(UChaosMoverSimulation::execK2_QueueInstantMovementEffect)
+{
+	Stack.StepCompiledIn<FStructProperty>(nullptr);
+	void* EffectPtr = Stack.MostRecentPropertyAddress;
+	FStructProperty* StructProp = CastField<FStructProperty>(Stack.MostRecentProperty);
+
+	P_FINISH;
+
+	P_NATIVE_BEGIN;
+
+	const bool bHasValidStructProp = StructProp && StructProp->Struct && StructProp->Struct->IsChildOf(FInstantMovementEffect::StaticStruct());
+
+	if (ensureMsgf((bHasValidStructProp && EffectPtr), TEXT("An invalid type (%s) was sent to a QueueInstantMovementEffect node. A struct derived from FInstantMovementEffect is required. No Movement Effect will be queued."),
+		StructProp ? *GetNameSafe(StructProp->Struct) : *Stack.MostRecentProperty->GetClass()->GetName()))
+	{
+		// Could we steal this instead of cloning? (move semantics)
+		FInstantMovementEffect* EffectAsBasePtr = reinterpret_cast<FInstantMovementEffect*>(EffectPtr);
+		FInstantMovementEffect* ClonedMove = EffectAsBasePtr->Clone();
+
+		P_THIS->QueueInstantMovementEffect(TSharedPtr<FInstantMovementEffect>(ClonedMove));
+	}
+
+	P_NATIVE_END;
+}
+
+
+void UChaosMoverSimulation::QueueInstantMovementEffect(TSharedPtr<FInstantMovementEffect> InstantMovementEffect)
+{
+	StateMachine.QueueInstantMovementEffect(InstantMovementEffect);
+}
diff --git a/Engine/Plugins/Experimental/ChaosMover/Source/ChaosMover/Private/Character/Effects/ChaosCharacterApplyVelocityEffect.cpp b/Engine/Plugins/Experimental/ChaosMover/Source/ChaosMover/Private/Character/Effects/ChaosCharacterApplyVelocityEffect.cpp
new file mode 100644
--- /dev/null
+++ b/Engine/Plugins/Experimental/ChaosMover/Source/ChaosMover/Private/Character/Effects/ChaosCharacterApplyVelocityEffect.cpp
@@ -0,0 +1,81 @@
+// Copyright Epic Games, Inc. All Rights Reserved.
+
+#include "ChaosMover/Character/Effects/ChaosCharacterApplyVelocityEffect.h"
+
+#include "Chaos/PhysicsObjectInternalInterface.h"
+#include "ChaosMover/ChaosMoverSimulationTypes.h"
+
+bool FChaosCharacterApplyVelocityEffect::ApplyMovementEffect_Async(FApplyMovementEffectParams_Async& ApplyEffectParams, FMoverSyncState& OutputState)
+{
+	check(ApplyEffectParams.StartState);
+	const FMoverDefaultSyncState* CurrentSyncState = ApplyEffectParams.StartState->SyncState.SyncStateCollection.FindDataByType<FMoverDefaultSyncState>();
+	const FChaosMoverSimulationDefaultInputs* SimInputs = ApplyEffectParams.StartState->AuxState.AuxStateCollection.FindDataByType<FChaosMoverSimulationDefaultInputs>();
+
+	if (!CurrentSyncState || !SimInputs)
+	{
+		return false;
+	}
+
+	// Get the position and orientation. Start by looking in the sync state.
+	// If not there, get from the particle corresponding to the updated component
+	FVector Position = CurrentSyncState->GetLocation_WorldSpace();
+	FRotator Orientation = CurrentSyncState->GetOrientation_WorldSpace();
+
+	FMoverDefaultSyncState& OutputSyncState = OutputState.SyncStateCollection.FindOrAddMutableDataByType<FMoverDefaultSyncState>();
+
+	FVector Velocity = CurrentSyncState->GetVelocity_WorldSpace();
+	switch (Mode)
+	{
+	case EChaosMoverVelocityEffectMode::Impulse:
+		{
+			Chaos::FReadPhysicsObjectInterface_Internal ReadInterface = Chaos::FPhysicsObjectInternalInterface::GetRead();
+			float Mass = ReadInterface.GetMass({ SimInputs->PhysicsObject });
+			if (Mass > UE_SMALL_NUMBER)
+			{
+				Velocity += (1.0f / Mass) * VelocityOrImpulseToApply;
+			}
+		}
+		break;
+	case EChaosMoverVelocityEffectMode::AdditiveVelocity:
+		Velocity += VelocityOrImpulseToApply;
+		break;
+	case EChaosMoverVelocityEffectMode::OverrideVelocity:
+		Velocity = VelocityOrImpulseToApply;
+		break;
+	default:
+		break;
+	}
+
+	OutputSyncState.SetTransforms_WorldSpace(Position,Orientation,Velocity);
+
+	return true;
+}
+
+FInstantMovementEffect* FChaosCharacterApplyVelocityEffect::Clone() const
+{
+	FChaosCharacterApplyVelocityEffect* CopyPtr = new FChaosCharacterApplyVelocityEffect(*this);
+	return CopyPtr;
+}
+
+void FChaosCharacterApplyVelocityEffect::NetSerialize(FArchive& Ar)
+{
+	Super::NetSerialize(Ar);
+
+	Ar << VelocityOrImpulseToApply;
+	Ar << Mode;
+}
+
+UScriptStruct* FChaosCharacterApplyVelocityEffect::GetScriptStruct() const
+{
+	return FChaosCharacterApplyVelocityEffect::StaticStruct();
+}
+
+FString FChaosCharacterApplyVelocityEffect::ToSimpleString() const
+{
+	return FString::Printf(TEXT("ChaosCharacterApplyVelocityEffect"));
+}
+
+void FChaosCharacterApplyVelocityEffect::AddReferencedObjects(FReferenceCollector& Collector)
+{
+	Super::AddReferencedObjects(Collector);
+}
\ No newline at end of file
diff --git a/Engine/Plugins/Experimental/ChaosMover/Source/ChaosMover/Private/Character/Modes/ChaosFlyingMode.cpp b/Engine/Plugins/Experimental/ChaosMover/Source/ChaosMover/Private/Character/Modes/ChaosFlyingMode.cpp
new file mode 100644
--- /dev/null
+++ b/Engine/Plugins/Experimental/ChaosMover/Source/ChaosMover/Private/Character/Modes/ChaosFlyingMode.cpp
@@ -0,0 +1,133 @@
+// Copyright Epic Games, Inc. All Rights Reserved.
+
+#include "ChaosMover/Character/Modes/ChaosFlyingMode.h"
+
+#include "ChaosMover/ChaosMoverLog.h"
+#include "ChaosMover/ChaosMoverSimulation.h"
+#include "ChaosMover/Character/Settings/SharedChaosCharacterMovementSettings.h"
+#include "ChaosMover/Character/Transitions/ChaosCharacterLandingCheck.h"
+#include "MoveLibrary/AirMovementUtils.h"
+#include "MoveLibrary/MovementUtils.h"
+
+#include UE_INLINE_GENERATED_CPP_BY_NAME(ChaosFlyingMode)
+
+UChaosFlyingMode::UChaosFlyingMode(const FObjectInitializer& ObjectInitializer)
+	: Super(ObjectInitializer)
+{
+	bSupportsAsync = true;
+
+	GameplayTags.AddTag(Mover_IsInAir);
+	GameplayTags.AddTag(Mover_IsFlying);
+
+	RadialForceLimit = 0.0f;
+	SwingTorqueLimit = 3000.0f;
+	TwistTorqueLimit = 0.0f;
+}
+
+void UChaosFlyingMode::GenerateMove_Implementation(const FMoverTickStartData& StartState, const FMoverTimeStep& TimeStep, FProposedMove& OutProposedMove) const
+{
+	const FChaosMoverSimulationDefaultInputs* DefaultSimInputs = StartState.AuxState.AuxStateCollection.FindDataByType<FChaosMoverSimulationDefaultInputs>();
+	const FCharacterDefaultInputs* CharacterInputs = StartState.InputCmd.InputCollection.FindDataByType<FCharacterDefaultInputs>();
+	if (!DefaultSimInputs || !CharacterInputs)
+	{
+		UE_LOG(LogChaosMover, Warning, TEXT("ChaosFallingMode requires FChaosMoverSimulationDefaultInputs and FCharacterDefaultInputs"));
+		return;
+	}
+
+	if (!Simulation)
+	{
+		UE_LOG(LogChaosMover, Warning, TEXT("No Simulation set on ChaosFallingMode"));
+		return;
+	}
+
+	const FMoverDefaultSyncState* StartingSyncState = StartState.SyncState.SyncStateCollection.FindDataByType<FMoverDefaultSyncState>();
+	check(StartingSyncState);
+
+	const float DeltaSeconds = TimeStep.StepMs * 0.001f;
+
+	FFreeMoveParams Params;
+	if (CharacterInputs)
+	{
+		Params.MoveInputType = CharacterInputs->GetMoveInputType();
+		Params.MoveInput = CharacterInputs->GetMoveInput_WorldSpace();
+		//const bool bMaintainInputMagnitude = true;
+		//Params.MoveInput = UPlanarConstraintUtils::ConstrainDirectionToPlane(FPlanarConstraint(), CharacterInputs->GetMoveInput_WorldSpace(), bMaintainInputMagnitude);
+	}
+	else
+	{
+		Params.MoveInputType = EMoveInputType::None;
+		Params.MoveInput = FVector::ZeroVector;
+	}
+
+	FRotator IntendedOrientation_WorldSpace;
+	// If there's no intent from input to change orientation, use the current orientation
+	if (!CharacterInputs || CharacterInputs->OrientationIntent.IsNearlyZero())
+	{
+		IntendedOrientation_WorldSpace = StartingSyncState->GetOrientation_WorldSpace();
+	}
+	else
+	{
+		IntendedOrientation_WorldSpace = CharacterInputs->GetOrientationIntentDir_WorldSpace().ToOrientationRotator();
+	}
+
+	FQuat WorldToGravityTransform = FQuat::FindBetweenNormals(FVector::UpVector, DefaultSimInputs->UpDir);
+
+	IntendedOrientation_WorldSpace = UMovementUtils::ApplyGravityToOrientationIntent(IntendedOrientation_WorldSpace, WorldToGravityTransform, bShouldCharacterRemainUpright);
+
+	Params.OrientationIntent = IntendedOrientation_WorldSpace;
+	Params.PriorVelocity = StartingSyncState->GetVelocity_WorldSpace();
+	Params.PriorOrientation = StartingSyncState->GetOrientation_WorldSpace();
+	Params.TurningRate = SharedSettings->TurningRate;
+	Params.TurningBoost = SharedSettings->TurningBoost;
+	Params.MaxSpeed = SharedSettings->MaxSpeed;
+	Params.Acceleration = SharedSettings->Acceleration;
+	Params.Deceleration = SharedSettings->Deceleration;
+	Params.DeltaSeconds = DeltaSeconds;
+	Params.WorldToGravityQuat = WorldToGravityTransform;
+	Params.bUseAccelerationForVelocityMove = SharedSettings->bUseAccelerationForVelocityMove;
+
+	OutProposedMove = UAirMovementUtils::ComputeControlledFreeMove(Params);
+
+	// Don't do floor checks in flying mode so just clear any previous results from the blackboard
+	UMoverBlackboard* SimBlackboard = Simulation->GetBlackboard_Mutable();
+	if (SimBlackboard)
+	{
+		SimBlackboard->Invalidate(CommonBlackboard::LastFloorResult);
+		SimBlackboard->Invalidate(CommonBlackboard::LastWaterResult);
+	}
+}
+
+void UChaosFlyingMode::SimulationTick_Implementation(const FSimulationTickParams& Params, FMoverTickEndData& OutputState)
+{
+	const FChaosMoverSimulationDefaultInputs* DefaultSimInputs = Params.StartState.AuxState.AuxStateCollection.FindDataByType<FChaosMoverSimulationDefaultInputs>();
+	if (!DefaultSimInputs)
+	{
+		UE_LOG(LogChaosMover, Warning, TEXT("ChaosFallingMode requires FChaosMoverSimulationDefaultInputs"));
+		return;
+	}
+
+	const FMoverDefaultSyncState* StartingSyncState = Params.StartState.SyncState.SyncStateCollection.FindDataByType<FMoverDefaultSyncState>();
+	check(StartingSyncState);
+
+	FProposedMove ProposedMove = Params.ProposedMove;
+	FMoverDefaultSyncState& OutputSyncState = OutputState.SyncState.SyncStateCollection.FindOrAddMutableDataByType<FMoverDefaultSyncState>();
+
+	const float DeltaSeconds = Params.TimeStep.StepMs * 0.001f;
+
+	FRotator TargetOrient = StartingSyncState->GetOrientation_WorldSpace();
+	if (!UMovementUtils::IsAngularVelocityZero(ProposedMove.AngularVelocity))
+	{
+		TargetOrient += (ProposedMove.AngularVelocity * DeltaSeconds);
+	}
+
+	// The physics simulation applies Z-only gravity acceleration via physics volumes, so we need to account for it here 
+	FVector TargetVel = ProposedMove.LinearVelocity - DefaultSimInputs->PhysicsObjectGravity * DefaultSimInputs->UpDir * DeltaSeconds;
+	FVector TargetPos = StartingSyncState->GetLocation_WorldSpace() + TargetVel * DeltaSeconds;
+
+	OutputState.MovementEndState.RemainingMs = 0.0f;
+	OutputSyncState.MoveDirectionIntent = ProposedMove.bHasDirIntent ? ProposedMove.DirectionIntent : FVector::ZeroVector;
+	OutputSyncState.SetTransforms_WorldSpace(
+		TargetPos,
+		TargetOrient,
+		TargetVel);
+}
\ No newline at end of file
diff --git a/Engine/Plugins/Experimental/ChaosMover/Source/ChaosMover/Private/Character/Modes/ChaosWalkingMode.cpp b/Engine/Plugins/Experimental/ChaosMover/Source/ChaosMover/Private/Character/Modes/ChaosWalkingMode.cpp
--- a/Engine/Plugins/Experimental/ChaosMover/Source/ChaosMover/Private/Character/Modes/ChaosWalkingMode.cpp
+++ b/Engine/Plugins/Experimental/ChaosMover/Source/ChaosMover/Private/Character/Modes/ChaosWalkingMode.cpp
@@ -8,6 +8,7 @@
 #include "ChaosMover/ChaosMoverSimulationTypes.h"
 #include "ChaosMover/Character/Settings/SharedChaosCharacterMovementSettings.h"
 #include "ChaosMover/Character/Transitions/ChaosCharacterFallingCheck.h"
+#include "ChaosMover/Character/Transitions/ChaosCharacterJumpCheck.h"
 #include "ChaosMover/Utilities/ChaosGroundMovementUtils.h"
 #include "ChaosMover/Utilities/ChaosMoverQueryUtils.h"
 #include "Math/UnitConversion.h"
@@ -32,6 +33,7 @@ UChaosWalkingMode::UChaosWalkingMode(const FObjectInitializer& ObjectInitializer
 	TwistTorqueLimit = 1500.0f;
 
 	Transitions.Add(CreateDefaultSubobject<UChaosCharacterFallingCheck>(TEXT("DefaultFallingCheck")));
+	Transitions.Add(CreateDefaultSubobject<UChaosCharacterJumpCheck>(TEXT("DefaultJumpCheck")));
 }
 
 void UChaosWalkingMode::UpdateConstraintSettings(Chaos::FCharacterGroundConstraintSettings& ConstraintSettings) const
diff --git a/Engine/Plugins/Experimental/ChaosMover/Source/ChaosMover/Private/Character/Transitions/ChaosCharacterFallingCheck.cpp b/Engine/Plugins/Experimental/ChaosMover/Source/ChaosMover/Private/Character/Transitions/ChaosCharacterFallingCheck.cpp
--- a/Engine/Plugins/Experimental/ChaosMover/Source/ChaosMover/Private/Character/Transitions/ChaosCharacterFallingCheck.cpp
+++ b/Engine/Plugins/Experimental/ChaosMover/Source/ChaosMover/Private/Character/Transitions/ChaosCharacterFallingCheck.cpp
@@ -88,7 +88,7 @@ FTransitionEvalResult UChaosCharacterFallingCheck::Evaluate_Implementation(const
 			const float ProjectedRelativeNormalVelocity = FloorResult.HitResult.ImpactNormal.Dot(ProjectedVelocity - ProjectedGroundVelocity);
 			const float ProjectedRelativeVerticalVelocity = SimInputs->UpDir.Dot(ProjectedVelocity - ProjectedGroundVelocity);
 			const float VerticalVelocityLimit = 2.0f / DeltaSeconds;
-			const bool bJumping = CharacterInputs && CharacterInputs->bIsJumpPressed;
+			const bool bJumping = CharacterInputs && CharacterInputs->bIsJumpJustPressed;
 			bool bIsLiftingOffSurface = false;
 			if (bJumping || ((ProjectedRelativeNormalVelocity > VerticalVelocityLimit) && bIsGroundMoving && (ProjectedRelativeVerticalVelocity > VerticalVelocityLimit)))
 			{
diff --git a/Engine/Plugins/Experimental/ChaosMover/Source/ChaosMover/Private/Character/Transitions/ChaosCharacterJumpCheck.cpp b/Engine/Plugins/Experimental/ChaosMover/Source/ChaosMover/Private/Character/Transitions/ChaosCharacterJumpCheck.cpp
new file mode 100644
--- /dev/null
+++ b/Engine/Plugins/Experimental/ChaosMover/Source/ChaosMover/Private/Character/Transitions/ChaosCharacterJumpCheck.cpp
@@ -0,0 +1,87 @@
+// Copyright Epic Games, Inc. All Rights Reserved.
+
+#include "ChaosMover/Character/Transitions/ChaosCharacterJumpCheck.h"
+
+#include "Chaos/ParticleHandle.h"
+#include "Chaos/PhysicsObjectInternalInterface.h"
+#include "ChaosMover/ChaosMoverLog.h"
+#include "ChaosMover/ChaosMoverSimulation.h"
+#include "ChaosMover/ChaosMoverSimulationTypes.h"
+#include "ChaosMover/Character/Effects/ChaosCharacterApplyVelocityEffect.h"
+#include "MoveLibrary/FloorQueryUtils.h"
+
+#include UE_INLINE_GENERATED_CPP_BY_NAME(ChaosCharacterJumpCheck)
+
+namespace Private
+{
+	static void ApplyImpulse(Chaos::FPBDRigidParticleHandle* Particle, const Chaos::FVec3& Impulse, const Chaos::FVec3& Location)
+	{
+		Chaos::FRigidTransform3 ComTransform = Particle->GetTransformXRCom();
+		const Chaos::FVec3 Offset = Location - ComTransform.GetLocation();
+		Particle->SetW(Particle->GetW() + Particle->InvI() * Offset.Cross(Impulse));
+		Particle->SetV(Particle->GetV() + Particle->InvM() * Impulse);
+	}
+}
+
+UChaosCharacterJumpCheck::UChaosCharacterJumpCheck(const FObjectInitializer& ObjectInitializer)
+	: Super(ObjectInitializer)
+{
+	bSupportsAsync = true;
+	bFirstSubStepOnly = true;
+
+	TransitionToMode = DefaultModeNames::Falling;
+}
+
+FTransitionEvalResult UChaosCharacterJumpCheck::Evaluate_Implementation(const FSimulationTickParams& Params) const
+{
+	FTransitionEvalResult EvalResult;
+
+	const FCharacterDefaultInputs* CharacterInputs = Params.StartState.InputCmd.InputCollection.FindDataByType<FCharacterDefaultInputs>();
+	if (CharacterInputs && CharacterInputs->bIsJumpJustPressed)
+	{
+		EvalResult.NextMode = TransitionToMode;
+	}
+
+	return EvalResult;
+}
+
+void UChaosCharacterJumpCheck::Trigger_Implementation(const FSimulationTickParams& Params)
+{
+	if (const FChaosMoverSimulationDefaultInputs* SimInputs = Params.StartState.AuxState.AuxStateCollection.FindDataByType<FChaosMoverSimulationDefaultInputs>())
+	{
+		FVector JumpVelocity = JumpUpwardsSpeed * SimInputs->UpDir;
+
+		TSharedPtr<FChaosCharacterApplyVelocityEffect> JumpMove = MakeShared<FChaosCharacterApplyVelocityEffect>();
+		JumpMove->VelocityOrImpulseToApply = JumpVelocity;
+		JumpMove->Mode = EChaosMoverVelocityEffectMode::AdditiveVelocity;
+
+		if (ensure(Simulation))
+		{
+			Simulation->QueueInstantMovementEffect(JumpMove);
+
+			// Apply an equal and opposite impulse to the ground
+			if (FractionalGroundReactionImpulse > 0.0f)
+			{
+				FFloorCheckResult FloorResult;
+				const UMoverBlackboard* Blackboard = Simulation->GetBlackboard();
+				if (Blackboard && Blackboard->TryGet(CommonBlackboard::LastFloorResult, FloorResult))
+				{
+					if (Chaos::FPhysicsObjectHandle GroundObject = FloorResult.HitResult.PhysicsObject)
+					{
+						Chaos::FWritePhysicsObjectInterface_Internal Interface = Chaos::FPhysicsObjectInternalInterface::GetWrite();
+						const float CharacterMass = Interface.GetMass({ SimInputs->PhysicsObject });
+						FVector ImpulseToApply = -FractionalGroundReactionImpulse * CharacterMass * JumpVelocity;
+						if (Chaos::FPBDRigidParticleHandle* GroundParticle = Interface.GetRigidParticle(GroundObject))
+						{
+							Private::ApplyImpulse(GroundParticle, ImpulseToApply, FloorResult.HitResult.ImpactPoint);
+						}
+					}
+				}
+			}
+		}
+	}
+	else
+	{
+		UE_LOG(LogChaosMover, Warning, TEXT("UChaosCharacterJumpCheck requires FChaosMoverSimulationDefaultInputs"));
+	}
+}
\ No newline at end of file
diff --git a/Engine/Plugins/Experimental/ChaosMover/Source/ChaosMover/Private/Character/Transitions/ChaosCharacterLandingCheck.cpp b/Engine/Plugins/Experimental/ChaosMover/Source/ChaosMover/Private/Character/Transitions/ChaosCharacterLandingCheck.cpp
--- a/Engine/Plugins/Experimental/ChaosMover/Source/ChaosMover/Private/Character/Transitions/ChaosCharacterLandingCheck.cpp
+++ b/Engine/Plugins/Experimental/ChaosMover/Source/ChaosMover/Private/Character/Transitions/ChaosCharacterLandingCheck.cpp
@@ -46,7 +46,7 @@ FTransitionEvalResult UChaosCharacterLandingCheck::Evaluate_Implementation(const
 			const float DeltaSeconds = Params.TimeStep.StepMs * 0.001f;
 			const FVector LinearVelocity = StartingSyncState->GetVelocity_WorldSpace();
 			const float VerticalVelocity = SimInputs->UpDir.Dot(LinearVelocity);
-			const bool bJumping = CharacterInputs && CharacterInputs->bIsJumpPressed;
+			const bool bJumping = CharacterInputs && CharacterInputs->bIsJumpJustPressed;
 			const bool bIsMovingUp = bJumping || (VerticalVelocity > 0.0f);
 
 			// Check for water landing
diff --git a/Engine/Plugins/Experimental/ChaosMover/Source/ChaosMover/Public/ChaosMover/ChaosMoverSimulation.h b/Engine/Plugins/Experimental/ChaosMover/Source/ChaosMover/Public/ChaosMover/ChaosMoverSimulation.h
--- a/Engine/Plugins/Experimental/ChaosMover/Source/ChaosMover/Public/ChaosMover/ChaosMoverSimulation.h
+++ b/Engine/Plugins/Experimental/ChaosMover/Source/ChaosMover/Public/ChaosMover/ChaosMoverSimulation.h
@@ -49,6 +49,13 @@ public:
 	UFUNCTION(BlueprintPure, Category = Mover)
 	const UBaseMovementMode* FindMovementModeByName(const FName& Name) const;
 
+	UFUNCTION(BlueprintCallable, CustomThunk, Category = Mover, meta = (CustomStructureParam = "EffectAsRawData", AllowAbstract = "false", DisplayName = "Queue Instant Movement Effect"))
+	void K2_QueueInstantMovementEffect(UPARAM(DisplayName = "Instant Movement Effect") const int32& EffectAsRawData);
+	DECLARE_FUNCTION(execK2_QueueInstantMovementEffect);
+
+	// Queue a Instant Movement Effect to take place at the end of this frame or start of the next subtick - whichever happens first
+	void QueueInstantMovementEffect(TSharedPtr<FInstantMovementEffect> Move);
+
 	struct FInitParams
 	{
 		TMap<FName, TWeakObjectPtr<UBaseMovementMode>> ModesToRegister;
@@ -86,6 +93,7 @@ protected:
 	virtual void OnPostSimulationTick(const FMoverTimeStep& TimeStep, UE::ChaosMover::FSimulationOutputData& OutputData);
 	virtual void OnModifyContacts(Chaos::FCollisionContactModifier& Modifier) const;
 
+	virtual void PreSimulationTickCharacter(const IChaosCharacterMovementModeInterface& CharacterMode, const FMoverTimeStep& TimeStep, const UE::ChaosMover::FSimulationInputData& InputData);
 	virtual void PostSimulationTickCharacter(const IChaosCharacterMovementModeInterface& CharacterMode, const FMoverTimeStep& TimeStep, UE::ChaosMover::FSimulationOutputData& OutputData);
 	virtual void PostSimulationTickCharacterConstraint(const IChaosCharacterConstraintMovementModeInterface& CharacterConstraintMode, const FMoverTimeStep& TimeStep, UE::ChaosMover::FSimulationOutputData& OutputData);
 
diff --git a/Engine/Plugins/Experimental/ChaosMover/Source/ChaosMover/Public/ChaosMover/Character/Effects/ChaosCharacterApplyVelocityEffect.h b/Engine/Plugins/Experimental/ChaosMover/Source/ChaosMover/Public/ChaosMover/Character/Effects/ChaosCharacterApplyVelocityEffect.h
new file mode 100644
--- /dev/null
+++ b/Engine/Plugins/Experimental/ChaosMover/Source/ChaosMover/Public/ChaosMover/Character/Effects/ChaosCharacterApplyVelocityEffect.h
@@ -0,0 +1,42 @@
+// Copyright Epic Games, Inc. All Rights Reserved.
+
+#pragma once
+
+#include "InstantMovementEffect.h"
+
+#include "ChaosCharacterApplyVelocityEffect.generated.h"
+
+UENUM(BlueprintType)
+enum class EChaosMoverVelocityEffectMode : uint8
+{
+	/** Apply as an additive impulse*/
+	Impulse,
+
+	/** Apply as an additive velocity */
+	AdditiveVelocity,
+
+	/** Apply as an override velocity */
+	OverrideVelocity,
+};
+
+/** Applies a velocity or impulse for a single tick */
+USTRUCT(BlueprintInternalUseOnly)
+struct CHAOSMOVER_API FChaosCharacterApplyVelocityEffect : public FInstantMovementEffect
+{
+	GENERATED_BODY()
+
+	// Velocity or impulse to apply to the actor.
+	UPROPERTY(EditAnywhere, BlueprintReadWrite, Category = Settings, meta = (ForceUnits = "cm/s"))
+	FVector VelocityOrImpulseToApply = FVector::ZeroVector;
+
+	// Controls whether to apply velocity or impulse and if the velocity will be additive
+	UPROPERTY(EditAnywhere, BlueprintReadWrite, Category = Settings)
+	EChaosMoverVelocityEffectMode Mode = EChaosMoverVelocityEffectMode::AdditiveVelocity;
+
+	virtual bool ApplyMovementEffect_Async(FApplyMovementEffectParams_Async& ApplyEffectParams, FMoverSyncState& OutputState);
+	virtual FInstantMovementEffect* Clone() const override;
+	virtual void NetSerialize(FArchive& Ar) override;
+	virtual UScriptStruct* GetScriptStruct() const override;
+	virtual FString ToSimpleString() const override;
+	virtual void AddReferencedObjects(class FReferenceCollector& Collector) override;
+};
\ No newline at end of file
diff --git a/Engine/Plugins/Experimental/ChaosMover/Source/ChaosMover/Public/ChaosMover/Character/Modes/ChaosFlyingMode.h b/Engine/Plugins/Experimental/ChaosMover/Source/ChaosMover/Public/ChaosMover/Character/Modes/ChaosFlyingMode.h
new file mode 100644
--- /dev/null
+++ b/Engine/Plugins/Experimental/ChaosMover/Source/ChaosMover/Public/ChaosMover/Character/Modes/ChaosFlyingMode.h
@@ -0,0 +1,22 @@
+// Copyright Epic Games, Inc. All Rights Reserved.
+
+#pragma once
+
+#include "ChaosMover/Character/Modes/ChaosCharacterMovementMode.h"
+
+#include "ChaosFlyingMode.generated.h"
+
+/**
+ * Chaos character flying mode
+ */
+UCLASS(Blueprintable, BlueprintType, EditInlineNew, DefaultToInstanced)
+class CHAOSMOVER_API UChaosFlyingMode : public UChaosCharacterMovementMode
+{
+	GENERATED_BODY()
+
+public:
+	UChaosFlyingMode(const FObjectInitializer& ObjectInitializer);
+
+	virtual void GenerateMove_Implementation(const FMoverTickStartData& StartState, const FMoverTimeStep& TimeStep, FProposedMove& OutProposedMove) const override;
+	virtual void SimulationTick_Implementation(const FSimulationTickParams& Params, FMoverTickEndData& OutputState) override;
+};
\ No newline at end of file
diff --git a/Engine/Plugins/Experimental/ChaosMover/Source/ChaosMover/Public/ChaosMover/Character/Transitions/ChaosCharacterJumpCheck.h b/Engine/Plugins/Experimental/ChaosMover/Source/ChaosMover/Public/ChaosMover/Character/Transitions/ChaosCharacterJumpCheck.h
new file mode 100644
--- /dev/null
+++ b/Engine/Plugins/Experimental/ChaosMover/Source/ChaosMover/Public/ChaosMover/Character/Transitions/ChaosCharacterJumpCheck.h
@@ -0,0 +1,36 @@
+// Copyright Epic Games, Inc. All Rights Reserved.
+
+#pragma once
+
+#include "ChaosMover/ChaosMovementModeTransition.h"
+
+#include "ChaosCharacterJumpCheck.generated.h"
+
+
+UCLASS(Blueprintable, BlueprintType, EditInlineNew, DefaultToInstanced)
+class CHAOSMOVER_API UChaosCharacterJumpCheck : public UChaosMovementModeTransition
+{
+	GENERATED_BODY()
+
+public:
+	UChaosCharacterJumpCheck(const FObjectInitializer& ObjectInitializer);
+
+	virtual FTransitionEvalResult Evaluate_Implementation(const FSimulationTickParams& Params) const override;
+	virtual void Trigger_Implementation(const FSimulationTickParams& Params) override;
+
+	/** Instantaneous speed induced in an actor upon jumping */
+	UPROPERTY(BlueprintReadWrite, EditAnywhere, Category = Parameters, meta = (ClampMin = "0", UIMin = "0", ForceUnits = "cm/s"))
+	float JumpUpwardsSpeed = 700.0f;
+
+	/**
+	* Controls how much of the jump impulse the character will apply to the ground.
+	* A value of 0 means no impulse will be applied to the ground.
+	* A value of 1 means that the full equal and opposite jump impulse will be applied.
+	*/
+	UPROPERTY(BlueprintReadWrite, EditAnywhere, Category = Parameters, meta = (ClampMin = "0", UIMin = "0", ClampMax = "1", UIMax = "1"))
+	float FractionalGroundReactionImpulse = 1.0f;
+
+	/** Name of movement mode to transition to when jumping */
+	UPROPERTY(BlueprintReadWrite, EditAnywhere, Category = Parameters)
+	FName TransitionToMode;
+};
\ No newline at end of file
