ChaosMover - Implement contact modification. This supports
- Disabling collision with ignore components and actors
- Walking mode: Disabling vertical contacts with the bottom of the character capsule
- Walking mode: Disabling contacts with the ground that are below the step height

diff --git a/Engine/Plugins/Experimental/ChaosMover/Source/ChaosMover/Private/Backends/ChaosMoverAsyncCallback.cpp b/Engine/Plugins/Experimental/ChaosMover/Source/ChaosMover/Private/Backends/ChaosMoverAsyncCallback.cpp
--- a/Engine/Plugins/Experimental/ChaosMover/Source/ChaosMover/Private/Backends/ChaosMoverAsyncCallback.cpp
+++ b/Engine/Plugins/Experimental/ChaosMover/Source/ChaosMover/Private/Backends/ChaosMoverAsyncCallback.cpp
@@ -51,14 +51,28 @@ namespace UE::ChaosMover
 	{
 		const FAsyncCallbackInput* AsyncInput = GetConsumerInput_Internal();
 
-		auto LambdaParallelUpdate = [&Modifier, &AsyncInput](int32 Idx) {
+		if (!AsyncInput)
+		{
+			return;
+		}
+
+		FMoverTimeStep TimeStep;
+		TimeStep.BaseSimTimeMs = GetSimTime_Internal() * 1000.0f;
+		TimeStep.StepMs = GetDeltaTime_Internal() * 1000.0f;
+		if (AsyncInput->PhysicsSolver)
+		{
+			TimeStep.ServerFrame = AsyncInput->PhysicsSolver->GetCurrentFrame() + AsyncInput->NetworkPhysicsTickOffset;
+			TimeStep.BaseSimTimeMs += AsyncInput->NetworkPhysicsTickOffset * TimeStep.StepMs;
+			TimeStep.bIsResimulating = AsyncInput->PhysicsSolver->GetEvolution()->IsResimming();
+		}
+
+		const FAsyncCallbackOutput& AsyncOutput = GetProducerOutputData_Internal();
+		for (int Idx = 0; Idx < AsyncInput->Backends.Num(); ++Idx)
+		{
 			if (TStrongObjectPtr<const UChaosMoverBackendComponent> Backend = AsyncInput->Backends[Idx].Pin())
 			{
-				Backend->GetSimulation()->ModifyContacts(Modifier);
+				Backend->GetSimulation()->ModifyContacts(TimeStep, AsyncInput->InputData[Idx], AsyncOutput.OutputData[Idx], Modifier);
 			}
-		};
-
-		const bool ForceSingleThread = UE::ChaosMover::CVars::bForceSingleThreadedPT;
-		Chaos::PhysicsParallelFor(AsyncInput->Backends.Num(), LambdaParallelUpdate, ForceSingleThread);
+		}
 	}
 }
\ No newline at end of file
diff --git a/Engine/Plugins/Experimental/ChaosMover/Source/ChaosMover/Private/Backends/ChaosMoverBackend.cpp b/Engine/Plugins/Experimental/ChaosMover/Source/ChaosMover/Private/Backends/ChaosMoverBackend.cpp
--- a/Engine/Plugins/Experimental/ChaosMover/Source/ChaosMover/Private/Backends/ChaosMoverBackend.cpp
+++ b/Engine/Plugins/Experimental/ChaosMover/Source/ChaosMover/Private/Backends/ChaosMoverBackend.cpp
@@ -203,7 +203,7 @@ void UChaosMoverBackendComponent::InitSimulation()
 	Params.ImmediateModeTransition = TWeakObjectPtr<UImmediateMovementModeTransition>(ImmediateModeTransition.Get());
 	Params.NullMovementMode = TWeakObjectPtr<UNullMovementMode>(NullMovementMode.Get());
 	Params.StartingMovementMode = MoverComp.StartingMovementMode;
-	Params.CharacterConstraintProxy = CharacterGroundConstraint->GetProxy<Chaos::FCharacterGroundConstraintProxy>();
+	Params.CharacterConstraintProxy = CharacterGroundConstraint ? CharacterGroundConstraint->GetProxy<Chaos::FCharacterGroundConstraintProxy>() : nullptr;
 	if (IPhysicsComponent* PhysicsComponent = Cast<IPhysicsComponent>(MoverComp.GetUpdatedComponent()))
 	{
 		if (Chaos::FPhysicsObjectHandle PhysObj = PhysicsComponent->GetPhysicsObjectByName(NAME_None))
diff --git a/Engine/Plugins/Experimental/ChaosMover/Source/ChaosMover/Private/ChaosMoverSimulation.cpp b/Engine/Plugins/Experimental/ChaosMover/Source/ChaosMover/Private/ChaosMoverSimulation.cpp
--- a/Engine/Plugins/Experimental/ChaosMover/Source/ChaosMover/Private/ChaosMoverSimulation.cpp
+++ b/Engine/Plugins/Experimental/ChaosMover/Source/ChaosMover/Private/ChaosMoverSimulation.cpp
@@ -4,6 +4,7 @@
 
 #include "Chaos/Character/CharacterGroundConstraint.h"
 #include "Chaos/Character/CharacterGroundConstraintContainer.h"
+#include "Chaos/ContactModification.h"
 #include "Chaos/PhysicsObjectInternalInterface.h"
 #include "ChaosMover/ChaosMovementMode.h"
 #include "ChaosMover/ChaosMovementModeTransition.h"
@@ -137,11 +138,67 @@ void UChaosMoverSimulation::SimulationTick(const FMoverTimeStep& TimeStep, const
 	OnPostSimulationTick(TimeStep, OutputData);
 }
 
-void UChaosMoverSimulation::ModifyContacts(Chaos::FCollisionContactModifier& Modifier) const
+void UChaosMoverSimulation::ModifyContacts(const FMoverTimeStep& TimeStep, const UE::ChaosMover::FSimulationInputData& InputData, const UE::ChaosMover::FSimulationOutputData& OutputData, Chaos::FCollisionContactModifier& Modifier) const
 {
 	Chaos::EnsureIsInPhysicsThreadContext();
 
-	OnModifyContacts(Modifier);
+	if (TStrongObjectPtr<const UBaseMovementMode> CurrentModePtr = StateMachine.GetCurrentMode().Pin())
+	{
+		if (const UChaosMovementMode* ChaosMode = Cast<UChaosMovementMode>(CurrentModePtr.Get()))
+		{
+			// Base contact modification
+			// Disable collisions for actors and components on the ignore list in the query params
+			if (ChaosMode->IgnoredCollisionMode == EChaosMoverIgnoredCollisionMode::DisableCollisionsWithIgnored)
+			{
+				Chaos::FGeometryParticleHandle* UpdatedParticle = nullptr;
+				const FChaosMoverSimulationDefaultInputs* SimInputs = InputData.LocalInputParams.AuxStateCollection.FindDataByType<FChaosMoverSimulationDefaultInputs>();
+				if (SimInputs)
+				{
+					Chaos::FReadPhysicsObjectInterface_Internal ReadInterface = Chaos::FPhysicsObjectInternalInterface::GetRead();
+					UpdatedParticle = ReadInterface.GetParticle(SimInputs->PhysicsObject);
+				}
+
+				if (!UpdatedParticle)
+				{
+					return;
+				}
+
+				for (Chaos::FContactPairModifier& PairModifier : Modifier.GetContacts(UpdatedParticle))
+				{
+					const int32 OtherIdx = UpdatedParticle == PairModifier.GetParticlePair()[0] ? 1 : 0;
+
+					if (const Chaos::FShapeInstance* Shape = PairModifier.GetShape(OtherIdx))
+					{
+						uint32 ComponentID = ChaosInterface::GetSimulationFilterData(*Shape).Word2;
+						if (SimInputs->CollisionQueryParams.GetIgnoredComponents().Contains(ComponentID))
+						{
+							PairModifier.Disable();
+							continue;
+						}
+
+						uint32 ActorID = ChaosInterface::GetQueryFilterData(*Shape).Word0;
+						if (SimInputs->CollisionQueryParams.GetIgnoredSourceObjects().Contains(ActorID))
+						{
+							PairModifier.Disable();
+							continue;
+						}
+
+						FMaskFilter ShapeMaskFilter = ChaosInterface::GetQueryFilterData(*Shape).Word3 >> (32u - NumExtraFilterBits);
+						if (SimInputs->CollisionQueryParams.IgnoreMask & ShapeMaskFilter)
+						{
+							PairModifier.Disable();
+							continue;
+						}
+					}
+				}
+			}
+
+			// Mode specific contact modification
+			ChaosMode->ModifyContacts(TimeStep, InputData, OutputData, Modifier);
+		}
+	}
+
+	OnModifyContacts(TimeStep, InputData, OutputData, Modifier);
 }
 
 void UChaosMoverSimulation::PreSimulationTickCharacter(const IChaosCharacterMovementModeInterface& CharacterMode, const FMoverTimeStep& TimeStep, const UE::ChaosMover::FSimulationInputData& InputData)
@@ -371,7 +428,7 @@ void UChaosMoverSimulation::OnPostSimulationTick(const FMoverTimeStep& TimeStep,
 	CurrentSyncState = OutputData.SyncState;
 }
 
-void UChaosMoverSimulation::OnModifyContacts(Chaos::FCollisionContactModifier& Modifier) const
+void UChaosMoverSimulation::OnModifyContacts(const FMoverTimeStep& TimeStep, const UE::ChaosMover::FSimulationInputData& InputData, const UE::ChaosMover::FSimulationOutputData& OutputData, Chaos::FCollisionContactModifier& Modifier) const
 {
 }
 
diff --git a/Engine/Plugins/Experimental/ChaosMover/Source/ChaosMover/Private/Character/Modes/ChaosWalkingMode.cpp b/Engine/Plugins/Experimental/ChaosMover/Source/ChaosMover/Private/Character/Modes/ChaosWalkingMode.cpp
--- a/Engine/Plugins/Experimental/ChaosMover/Source/ChaosMover/Private/Character/Modes/ChaosWalkingMode.cpp
+++ b/Engine/Plugins/Experimental/ChaosMover/Source/ChaosMover/Private/Character/Modes/ChaosWalkingMode.cpp
@@ -3,6 +3,8 @@
 #include "ChaosMover/Character/Modes/ChaosWalkingMode.h"
 
 #include "Chaos/Character/CharacterGroundConstraintSettings.h"
+#include "Chaos/ContactModification.h"
+#include "Chaos/PhysicsObjectInternalInterface.h"
 #include "ChaosMover/ChaosMoverLog.h"
 #include "ChaosMover/ChaosMoverSimulation.h"
 #include "ChaosMover/ChaosMoverSimulationTypes.h"
@@ -404,3 +406,76 @@ void UChaosWalkingMode::GetFloorAndCheckMovement(const FMoverDefaultSyncState& S
 		OutDeltaPos = FVector::ZeroVector;
 	}
 }
+
+void UChaosWalkingMode::ModifyContacts(const FMoverTimeStep& TimeStep, const UE::ChaosMover::FSimulationInputData& InputData, const UE::ChaosMover::FSimulationOutputData& OutputData, Chaos::FCollisionContactModifier& Modifier) const
+{
+	check(Simulation);
+
+	// Get the updated (character) particle
+	Chaos::FGeometryParticleHandle* UpdatedParticle = nullptr;
+	const FChaosMoverSimulationDefaultInputs* SimInputs = InputData.LocalInputParams.AuxStateCollection.FindDataByType<FChaosMoverSimulationDefaultInputs>();
+	if (SimInputs)
+	{
+		Chaos::FReadPhysicsObjectInterface_Internal ReadInterface = Chaos::FPhysicsObjectInternalInterface::GetRead();
+		UpdatedParticle = ReadInterface.GetParticle(SimInputs->PhysicsObject);
+	}
+
+	if (!UpdatedParticle)
+	{
+		return;
+	}
+
+	// Try and find the ground particle if there is one in the latest floor result
+	Chaos::FGeometryParticleHandle* GroundParticle = nullptr;
+	FFloorCheckResult FloorResult;
+	FloorResult.FloorDist = 1.0e10;
+	if (const UMoverBlackboard* SimBlackboard = Simulation->GetBlackboard())
+	{
+		SimBlackboard->TryGet(CommonBlackboard::LastFloorResult, FloorResult);
+		Chaos::FReadPhysicsObjectInterface_Internal Interface = Chaos::FPhysicsObjectInternalInterface::GetRead();
+		GroundParticle = Interface.GetParticle(FloorResult.HitResult.PhysicsObject);
+	}
+
+	const float CharacterHeight = UpdatedParticle->GetX().Dot(SimInputs->UpDir);
+	const float EndCapHeight = CharacterHeight - SimInputs->PawnCollisionHalfHeight + SimInputs->PawnCollisionRadius;
+	constexpr float CosThetaMax = 0.707f;
+
+	float MinContactHeightStepUps = CharacterHeight - 1.0e10f;
+	const float StepDistance = FMath::Abs(GetTargetHeight() - FloorResult.FloorDist);
+	if (StepDistance >= UE_KINDA_SMALL_NUMBER)
+	{
+		MinContactHeightStepUps = CharacterHeight - GetTargetHeight() + SharedSettings->MaxStepHeight;
+	}
+
+	for (Chaos::FContactPairModifier& PairModifier : Modifier.GetContacts(UpdatedParticle))
+	{
+		const int32 CharacterIdx = UpdatedParticle == PairModifier.GetParticlePair()[0] ? 0 : 1;
+		Chaos::FGeometryParticleHandle* OtherParticle = PairModifier.GetOtherParticle(UpdatedParticle);
+		if (!OtherParticle)
+		{
+			continue;
+		}
+		const bool bOtherParticleIsGround = OtherParticle == GroundParticle;
+
+		for (int32 Idx = 0; Idx < PairModifier.GetNumContacts(); ++Idx)
+		{
+			Chaos::FVec3 Point0, Point1;
+			PairModifier.GetWorldContactLocations(Idx, Point0, Point1);
+			Chaos::FVec3 CharacterPoint = CharacterIdx == 0 ? Point0 : Point1;
+
+			Chaos::FVec3 ContactNormal = PairModifier.GetWorldNormal(Idx);
+			if ((ContactNormal.Z > CosThetaMax) && CharacterPoint.Z < EndCapHeight)
+			{
+				// Disable any nearly vertical contact with the end cap of the capsule
+				// This will be handled by the character ground constraint
+				PairModifier.SetContactPointDisabled(Idx);
+			}
+			else if (bOtherParticleIsGround && CharacterPoint.Z < MinContactHeightStepUps)
+			{
+				// In the case of steps ups disable all contacts below the max step height
+				PairModifier.SetContactPointDisabled(Idx);
+			}
+		}
+	}
+}
+
diff --git a/Engine/Plugins/Experimental/ChaosMover/Source/ChaosMover/Public/ChaosMover/ChaosMovementMode.h b/Engine/Plugins/Experimental/ChaosMover/Source/ChaosMover/Public/ChaosMover/ChaosMovementMode.h
--- a/Engine/Plugins/Experimental/ChaosMover/Source/ChaosMover/Public/ChaosMover/ChaosMovementMode.h
+++ b/Engine/Plugins/Experimental/ChaosMover/Source/ChaosMover/Public/ChaosMover/ChaosMovementMode.h
@@ -2,11 +2,23 @@
 
 #pragma once
 
+#include "ChaosMover/ChaosMoverSimulationTypes.h"
 #include "MovementMode.h"
 
 #include "ChaosMovementMode.generated.h"
 
 class UChaosMoverSimulation;
+namespace Chaos
+{
+	class FCollisionContactModifier;
+}
+
+UENUM()
+enum class EChaosMoverIgnoredCollisionMode : uint8
+{
+	EnableCollisionsWithIgnored,
+	DisableCollisionsWithIgnored,
+};
 
 /**
  * Base class for all Chaos movement modes
@@ -27,6 +39,13 @@ public:
 
 	void SetSimulation(UChaosMoverSimulation* InSimulation);
 
+	virtual void ModifyContacts(const FMoverTimeStep& TimeStep, const UE::ChaosMover::FSimulationInputData& InputData, const UE::ChaosMover::FSimulationOutputData& OutputData, Chaos::FCollisionContactModifier& Modifier) const
+	{
+	}
+
+	UPROPERTY(EditAnywhere, Category = Collision)
+	EChaosMoverIgnoredCollisionMode IgnoredCollisionMode = EChaosMoverIgnoredCollisionMode::DisableCollisionsWithIgnored;
+
 protected:
 	UChaosMoverSimulation* Simulation;
 };
\ No newline at end of file
diff --git a/Engine/Plugins/Experimental/ChaosMover/Source/ChaosMover/Public/ChaosMover/ChaosMoverSimulation.h b/Engine/Plugins/Experimental/ChaosMover/Source/ChaosMover/Public/ChaosMover/ChaosMoverSimulation.h
--- a/Engine/Plugins/Experimental/ChaosMover/Source/ChaosMover/Public/ChaosMover/ChaosMoverSimulation.h
+++ b/Engine/Plugins/Experimental/ChaosMover/Source/ChaosMover/Public/ChaosMover/ChaosMoverSimulation.h
@@ -74,7 +74,7 @@ public:
 	void Init(const FInitParams& InitParams);
 	void Deinit();
 	void SimulationTick(const FMoverTimeStep& TimeStep, const UE::ChaosMover::FSimulationInputData& InputData, UE::ChaosMover::FSimulationOutputData& OutputData);
-	void ModifyContacts(Chaos::FCollisionContactModifier& Modifier) const;
+	void ModifyContacts(const FMoverTimeStep& TimeStep, const UE::ChaosMover::FSimulationInputData& InputData, const UE::ChaosMover::FSimulationOutputData& OutputData, Chaos::FCollisionContactModifier& Modifier) const;
 
 	const TArray<UE::ChaosMover::FMovementModeChangeEvent>& GetQueuedModeChangeEvents() const;
 	void ClearQueuedModeChangeEvents();
@@ -91,7 +91,7 @@ protected:
 	virtual void OnPreSimulationTick(const FMoverTimeStep& TimeStep, const UE::ChaosMover::FSimulationInputData& InputData);
 	virtual void OnSimulationTick(const FMoverTimeStep& TimeStep, const UE::ChaosMover::FSimulationInputData& InputData, UE::ChaosMover::FSimulationOutputData& OutputData);
 	virtual void OnPostSimulationTick(const FMoverTimeStep& TimeStep, UE::ChaosMover::FSimulationOutputData& OutputData);
-	virtual void OnModifyContacts(Chaos::FCollisionContactModifier& Modifier) const;
+	virtual void OnModifyContacts(const FMoverTimeStep& TimeStep, const UE::ChaosMover::FSimulationInputData& InputData, const UE::ChaosMover::FSimulationOutputData& OutputData, Chaos::FCollisionContactModifier& Modifier) const;
 
 	virtual void PreSimulationTickCharacter(const IChaosCharacterMovementModeInterface& CharacterMode, const FMoverTimeStep& TimeStep, const UE::ChaosMover::FSimulationInputData& InputData);
 	virtual void PostSimulationTickCharacter(const IChaosCharacterMovementModeInterface& CharacterMode, const FMoverTimeStep& TimeStep, UE::ChaosMover::FSimulationOutputData& OutputData);
diff --git a/Engine/Plugins/Experimental/ChaosMover/Source/ChaosMover/Public/ChaosMover/Character/Modes/ChaosWalkingMode.h b/Engine/Plugins/Experimental/ChaosMover/Source/ChaosMover/Public/ChaosMover/Character/Modes/ChaosWalkingMode.h
--- a/Engine/Plugins/Experimental/ChaosMover/Source/ChaosMover/Public/ChaosMover/Character/Modes/ChaosWalkingMode.h
+++ b/Engine/Plugins/Experimental/ChaosMover/Source/ChaosMover/Public/ChaosMover/Character/Modes/ChaosWalkingMode.h
@@ -53,6 +53,8 @@ public:
 	virtual void GenerateMove_Implementation(const FMoverTickStartData& StartState, const FMoverTimeStep& TimeStep, FProposedMove& OutProposedMove) const override;
 	virtual void SimulationTick_Implementation(const FSimulationTickParams& Params, FMoverTickEndData& OutputState) override;
 
+	virtual void ModifyContacts(const FMoverTimeStep& TimeStep, const UE::ChaosMover::FSimulationInputData& InputData, const UE::ChaosMover::FSimulationOutputData& OutputData, Chaos::FCollisionContactModifier& Modifier) const override;
+
 protected:
 	virtual bool CanStepUpOnHitSurface(const FFloorCheckResult& FloorResult) const;
 	virtual void GetFloorAndCheckMovement(const FMoverDefaultSyncState& SyncState, const FProposedMove& ProposedMove, const FChaosMoverSimulationDefaultInputs& DefaultSimInputs, float DeltaSeconds, FFloorCheckResult& FloorResult, FWaterCheckResult& WaterResult, FVector& OutDeltaPos) const;
