Chaos Mover: Adding local sim input and internal sim state FMoverDataCollections to the chaos simulation. It hosts all state that is not sent across the network. Wrapping TargetDeltaRotation and TargetDeltaFacing into a new struct (FChaosMoverCharacterSimState) and moving it and FChaosMoverSimulationDefaultInputs to those instead of Aux.

diff --git a/Engine/Plugins/Experimental/ChaosMover/Source/ChaosMover/Private/Backends/ChaosMoverBackend.cpp b/Engine/Plugins/Experimental/ChaosMover/Source/ChaosMover/Private/Backends/ChaosMoverBackend.cpp
--- a/Engine/Plugins/Experimental/ChaosMover/Source/ChaosMover/Private/Backends/ChaosMoverBackend.cpp
+++ b/Engine/Plugins/Experimental/ChaosMover/Source/ChaosMover/Private/Backends/ChaosMoverBackend.cpp
@@ -331,7 +331,7 @@ void UChaosMoverBackendComponent::ProduceInputData(int32 PhysicsStep, int32 NumS
 	UMoverComponent& MoverComp = GetMoverComponent();
 
 	// Add default simulation input data
-	FChaosMoverSimulationDefaultInputs& SimInputs = InputData.LocalInputParams.AuxStateCollection.FindOrAddMutableDataByType<FChaosMoverSimulationDefaultInputs>();
+	FChaosMoverSimulationDefaultInputs& SimInputs = Simulation->GetLocalSimInput_Mutable().FindOrAddMutableDataByType<FChaosMoverSimulationDefaultInputs>();
 	SimInputs.Gravity = MoverComp.GetGravityAcceleration();
 	SimInputs.UpDir = MoverComp.GetUpDirection();
 	SimInputs.OwningActor = GetOwner();
@@ -407,7 +407,7 @@ void UChaosMoverBackendComponent::FinalizeFrame(double ResultsTimeInSeconds)
 	if (FChaosVisualDebuggerTrace::IsTracing())
 	{
 		const UE::ChaosMover::FSimulationOutputData& LatestOutput = SimOutputRecord.GetLatest();
-		FMoverCVDRuntimeTrace::TraceMoverData(&MoverComp, &InterpolatedOutput.LastUsedInputCmd, &InterpolatedOutput.SyncState);
+		FMoverCVDRuntimeTrace::TraceMoverData(&MoverComp, &InterpolatedOutput.LastUsedInputCmd, &InterpolatedOutput.SyncState, &Simulation->GetLocalSimInput());
 	}
 #endif
 
diff --git a/Engine/Plugins/Experimental/ChaosMover/Source/ChaosMover/Private/ChaosMoverSimulation.cpp b/Engine/Plugins/Experimental/ChaosMover/Source/ChaosMover/Private/ChaosMoverSimulation.cpp
--- a/Engine/Plugins/Experimental/ChaosMover/Source/ChaosMover/Private/ChaosMoverSimulation.cpp
+++ b/Engine/Plugins/Experimental/ChaosMover/Source/ChaosMover/Private/ChaosMoverSimulation.cpp
@@ -29,6 +29,29 @@ UChaosMoverSimulation::UChaosMoverSimulation()
 	Blackboard = CreateDefaultSubobject<UMoverBlackboard>(TEXT("ChaosMoverSimulationBlackboard"));
 }
 
+const UMoverBlackboard* UChaosMoverSimulation::GetBlackboard() const
+{
+	return Blackboard;
+}
+
+UMoverBlackboard* UChaosMoverSimulation::GetBlackboard_Mutable()
+{
+	return Blackboard;
+}
+
+const FMoverDataCollection& UChaosMoverSimulation::GetLocalSimInput() const
+{
+	return LocalSimInput;
+}
+
+FMoverDataCollection& UChaosMoverSimulation::GetLocalSimInput_Mutable()
+{
+	// Only the Gameplay Thread is allowed to write to the local simulation input data collection
+	Chaos::EnsureIsInGameThreadContext();
+
+	return LocalSimInput;
+}
+
 const UBaseMovementMode* UChaosMoverSimulation::GetCurrentMovementMode() const
 {
 	return StateMachine.GetCurrentMode().Get();
@@ -262,40 +285,30 @@ void UChaosMoverSimulation::OnSimulationTick(const FMoverTimeStep& TimeStep, con
 	// Copy the sync state locally and to the output data
 	OutputData.SyncState = TickEndData.SyncState;
 	OutputData.LastUsedInputCmd = InputData.InputCmd;
-
-	// Store the delta position and rotation locally
-	if (Chaos::FPBDRigidParticleHandle* ParticleHandle = UChaosMoverSimulation::GetControlledParticle())
-	{
-		const FMoverDefaultSyncState* DefaultSyncState = OutputData.SyncState.SyncStateCollection.FindDataByType<FMoverDefaultSyncState>();
-		check(DefaultSyncState);
-
-		FRotator DeltaRotation = DefaultSyncState->GetOrientation_WorldSpace() - FRotator(ParticleHandle->GetR());
-		FRotator Winding, Remainder;
-		DeltaRotation.GetWindingAndRemainder(Winding, Remainder);
-
-		TargetDeltaPosition = DefaultSyncState->GetLocation_WorldSpace() - ParticleHandle->GetX();
-		TargetDeltaFacing = FMath::DegreesToRadians(Remainder.Yaw);
-	}
-	else
-	{
-		TargetDeltaPosition = FVector::ZeroVector;
-		TargetDeltaFacing = 0.0f;
-	}
 }
 
 void UChaosMoverSimulation::PostSimulationTickCharacter(const IChaosCharacterMovementModeInterface& CharacterMode, const FMoverTimeStep& TimeStep, UE::ChaosMover::FSimulationOutputData& OutputData)
 {
+	FChaosMoverCharacterSimState& CharacterSimState = InternalSimData.FindOrAddMutableDataByType<FChaosMoverCharacterSimState>();
 	const FMoverDefaultSyncState* DefaultSyncState = OutputData.SyncState.SyncStateCollection.FindDataByType<FMoverDefaultSyncState>();
 	check(DefaultSyncState);
 
-	// Set the particle velocity to hit the target sync state transform
 	if (Chaos::FPBDRigidParticleHandle* ParticleHandle = UChaosMoverSimulation::GetControlledParticle())
 	{
+		// Store the delta position and rotation locally
+		FRotator DeltaRotation = DefaultSyncState->GetOrientation_WorldSpace() - FRotator(ParticleHandle->GetR());
+		FRotator Winding, Remainder;
+		DeltaRotation.GetWindingAndRemainder(Winding, Remainder);
+
+		CharacterSimState.TargetDeltaPosition = DefaultSyncState->GetLocation_WorldSpace() - ParticleHandle->GetX();
+		CharacterSimState.TargetDeltaFacing = FMath::DegreesToRadians(Remainder.Yaw);
+
+		// Set the particle velocity to hit the target sync state transform
 		ParticleHandle->SetV(DefaultSyncState->GetVelocity_WorldSpace());
 		const float DeltaTimeSeconds = TimeStep.StepMs * 0.001f;
 		if (DeltaTimeSeconds > UE_SMALL_NUMBER)
 		{
-			FVector AngularVelocity = (TargetDeltaFacing / DeltaTimeSeconds) * FVector::UpVector;
+			FVector AngularVelocity = (CharacterSimState.TargetDeltaFacing / DeltaTimeSeconds) * FVector::UpVector;
 
 			if (!CharacterMode.ShouldCharacterRemainUpright())
 			{
@@ -307,6 +320,11 @@ void UChaosMoverSimulation::PostSimulationTickCharacter(const IChaosCharacterMov
 			ParticleHandle->SetW(AngularVelocity);
 		}
 	}
+	else
+	{
+		CharacterSimState.TargetDeltaPosition = FVector::ZeroVector;
+		CharacterSimState.TargetDeltaFacing = 0.0f;
+	}
 }
 
 void UChaosMoverSimulation::PostSimulationTickCharacterConstraint(const IChaosCharacterConstraintMovementModeInterface& CharacterConstraintMode, const FMoverTimeStep& TimeStep, UE::ChaosMover::FSimulationOutputData& OutputData)
@@ -379,10 +397,13 @@ void UChaosMoverSimulation::PostSimulationTickCharacterConstraint(const IChaosCh
 			WalkableSlopeCosine = 2.0f;
 		}
 
+		FChaosMoverCharacterSimState* CharacterSimState = InternalSimData.FindMutableDataByType<FChaosMoverCharacterSimState>();
+		check(CharacterSimState);
+
 		ConstraintHandle->SetData({
 			FloorResult.HitResult.ImpactNormal,
-			TargetDeltaPosition,
-			TargetDeltaFacing,
+			CharacterSimState->TargetDeltaPosition,
+			CharacterSimState->TargetDeltaFacing,
 			FloorResult.FloorDist,
 			WalkableSlopeCosine
 			});
diff --git a/Engine/Plugins/Experimental/ChaosMover/Source/ChaosMover/Private/Character/Modes/ChaosFallingMode.cpp b/Engine/Plugins/Experimental/ChaosMover/Source/ChaosMover/Private/Character/Modes/ChaosFallingMode.cpp
--- a/Engine/Plugins/Experimental/ChaosMover/Source/ChaosMover/Private/Character/Modes/ChaosFallingMode.cpp
+++ b/Engine/Plugins/Experimental/ChaosMover/Source/ChaosMover/Private/Character/Modes/ChaosFallingMode.cpp
@@ -40,17 +40,17 @@ UChaosFallingMode::UChaosFallingMode(const FObjectInitializer& ObjectInitializer
 
 void UChaosFallingMode::GenerateMove_Implementation(const FMoverTickStartData& StartState, const FMoverTimeStep& TimeStep, FProposedMove& OutProposedMove) const
 {
-	const FChaosMoverSimulationDefaultInputs* DefaultSimInputs = StartState.AuxState.AuxStateCollection.FindDataByType<FChaosMoverSimulationDefaultInputs>();
-	const FCharacterDefaultInputs* CharacterInputs = StartState.InputCmd.InputCollection.FindDataByType<FCharacterDefaultInputs>();
-	if (!DefaultSimInputs || !CharacterInputs)
+	if (!Simulation)
 	{
-		UE_LOG(LogChaosMover, Warning, TEXT("ChaosFallingMode requires FChaosMoverSimulationDefaultInputs and FCharacterDefaultInputs"));
+		UE_LOG(LogChaosMover, Warning, TEXT("No Simulation set on ChaosFallingMode"));
 		return;
 	}
 
-	if (!Simulation)
+	const FChaosMoverSimulationDefaultInputs* DefaultSimInputs = Simulation->GetLocalSimInput().FindDataByType<FChaosMoverSimulationDefaultInputs>();
+	const FCharacterDefaultInputs* CharacterInputs = StartState.InputCmd.InputCollection.FindDataByType<FCharacterDefaultInputs>();
+	if (!DefaultSimInputs || !CharacterInputs)
 	{
-		UE_LOG(LogChaosMover, Warning, TEXT("No Simulation set on ChaosFallingMode"));
+		UE_LOG(LogChaosMover, Warning, TEXT("ChaosFallingMode requires FChaosMoverSimulationDefaultInputs and FCharacterDefaultInputs"));
 		return;
 	}
 
@@ -170,7 +170,13 @@ void UChaosFallingMode::GenerateMove_Implementation(const FMoverTickStartData& S
 
 void UChaosFallingMode::SimulationTick_Implementation(const FSimulationTickParams& Params, FMoverTickEndData& OutputState)
 {
-	const FChaosMoverSimulationDefaultInputs* DefaultSimInputs = Params.StartState.AuxState.AuxStateCollection.FindDataByType<FChaosMoverSimulationDefaultInputs>();
+	if (!Simulation)
+	{
+		UE_LOG(LogChaosMover, Warning, TEXT("No Simulation set on ChaosFallingMode"));
+		return;
+	}
+
+	const FChaosMoverSimulationDefaultInputs* DefaultSimInputs = Simulation->GetLocalSimInput().FindDataByType<FChaosMoverSimulationDefaultInputs>();
 	if (!DefaultSimInputs)
 	{
 		UE_LOG(LogChaosMover, Warning, TEXT("ChaosFallingMode requires FChaosMoverSimulationDefaultInputs"));
diff --git a/Engine/Plugins/Experimental/ChaosMover/Source/ChaosMover/Private/Character/Modes/ChaosWalkingMode.cpp b/Engine/Plugins/Experimental/ChaosMover/Source/ChaosMover/Private/Character/Modes/ChaosWalkingMode.cpp
--- a/Engine/Plugins/Experimental/ChaosMover/Source/ChaosMover/Private/Character/Modes/ChaosWalkingMode.cpp
+++ b/Engine/Plugins/Experimental/ChaosMover/Source/ChaosMover/Private/Character/Modes/ChaosWalkingMode.cpp
@@ -49,7 +49,7 @@ void UChaosWalkingMode::UpdateConstraintSettings(Chaos::FCharacterGroundConstrai
 
 void UChaosWalkingMode::GenerateMove_Implementation(const FMoverTickStartData& StartState, const FMoverTimeStep& TimeStep, FProposedMove& OutProposedMove) const
 {
-	const FChaosMoverSimulationDefaultInputs* DefaultSimInputs = StartState.AuxState.AuxStateCollection.FindDataByType<FChaosMoverSimulationDefaultInputs>();
+	const FChaosMoverSimulationDefaultInputs* DefaultSimInputs = Simulation->GetLocalSimInput().FindDataByType<FChaosMoverSimulationDefaultInputs>();
 	const FCharacterDefaultInputs* CharacterInputs = StartState.InputCmd.InputCollection.FindDataByType<FCharacterDefaultInputs>();
 	if (!DefaultSimInputs || !CharacterInputs)
 	{
@@ -162,7 +162,7 @@ void UChaosWalkingMode::GenerateMove_Implementation(const FMoverTickStartData& S
 
 void UChaosWalkingMode::SimulationTick_Implementation(const FSimulationTickParams& Params, FMoverTickEndData& OutputState)
 {
-	const FChaosMoverSimulationDefaultInputs* DefaultSimInputs = Params.StartState.AuxState.AuxStateCollection.FindDataByType<FChaosMoverSimulationDefaultInputs>();
+	const FChaosMoverSimulationDefaultInputs* DefaultSimInputs = Simulation->GetLocalSimInput().FindDataByType<FChaosMoverSimulationDefaultInputs>();
 	const FCharacterDefaultInputs* CharacterInputs = Params.StartState.InputCmd.InputCollection.FindDataByType<FCharacterDefaultInputs>();
 	if (!DefaultSimInputs || !CharacterInputs)
 	{
diff --git a/Engine/Plugins/Experimental/ChaosMover/Source/ChaosMover/Private/Character/Transitions/ChaosCharacterFallingCheck.cpp b/Engine/Plugins/Experimental/ChaosMover/Source/ChaosMover/Private/Character/Transitions/ChaosCharacterFallingCheck.cpp
--- a/Engine/Plugins/Experimental/ChaosMover/Source/ChaosMover/Private/Character/Transitions/ChaosCharacterFallingCheck.cpp
+++ b/Engine/Plugins/Experimental/ChaosMover/Source/ChaosMover/Private/Character/Transitions/ChaosCharacterFallingCheck.cpp
@@ -51,7 +51,7 @@ FTransitionEvalResult UChaosCharacterFallingCheck::Evaluate_Implementation(const
 
 	const FMoverTickStartData& StartState = Params.StartState;
 
-	if (const FChaosMoverSimulationDefaultInputs* SimInputs = StartState.AuxState.AuxStateCollection.FindDataByType<FChaosMoverSimulationDefaultInputs>())
+	if (const FChaosMoverSimulationDefaultInputs* SimInputs = Simulation->GetLocalSimInput().FindDataByType<FChaosMoverSimulationDefaultInputs>())
 	{
 		UMoverBlackboard* Blackboard = Simulation->GetBlackboard_Mutable();
 		check(Blackboard);
diff --git a/Engine/Plugins/Experimental/ChaosMover/Source/ChaosMover/Private/Character/Transitions/ChaosCharacterLandingCheck.cpp b/Engine/Plugins/Experimental/ChaosMover/Source/ChaosMover/Private/Character/Transitions/ChaosCharacterLandingCheck.cpp
--- a/Engine/Plugins/Experimental/ChaosMover/Source/ChaosMover/Private/Character/Transitions/ChaosCharacterLandingCheck.cpp
+++ b/Engine/Plugins/Experimental/ChaosMover/Source/ChaosMover/Private/Character/Transitions/ChaosCharacterLandingCheck.cpp
@@ -29,7 +29,7 @@ FTransitionEvalResult UChaosCharacterLandingCheck::Evaluate_Implementation(const
 
 	const FMoverTickStartData& StartState = Params.StartState;
 
-	if (const FChaosMoverSimulationDefaultInputs* SimInputs = StartState.AuxState.AuxStateCollection.FindDataByType<FChaosMoverSimulationDefaultInputs>())
+	if (const FChaosMoverSimulationDefaultInputs* SimInputs = Simulation->GetLocalSimInput().FindDataByType<FChaosMoverSimulationDefaultInputs>())
 	{
 		const UMoverBlackboard* Blackboard = Simulation->GetBlackboard();
 		check(Blackboard);
diff --git a/Engine/Plugins/Experimental/ChaosMover/Source/ChaosMover/Public/ChaosMover/ChaosMoverSimulation.h b/Engine/Plugins/Experimental/ChaosMover/Source/ChaosMover/Public/ChaosMover/ChaosMoverSimulation.h
--- a/Engine/Plugins/Experimental/ChaosMover/Source/ChaosMover/Public/ChaosMover/ChaosMoverSimulation.h
+++ b/Engine/Plugins/Experimental/ChaosMover/Source/ChaosMover/Public/ChaosMover/ChaosMoverSimulation.h
@@ -32,16 +32,16 @@ public:
 	UChaosMoverSimulation();
 
 	UFUNCTION(BlueprintPure, Category = Mover)
-	const UMoverBlackboard* GetBlackboard() const
-	{
-		return Blackboard;
-	}
+	const UMoverBlackboard* GetBlackboard() const;
 
 	UFUNCTION(BlueprintPure, Category = Mover)
-	UMoverBlackboard* GetBlackboard_Mutable()
-	{
-		return Blackboard;
-	}
+	UMoverBlackboard* GetBlackboard_Mutable();
+
+	UFUNCTION(BlueprintPure, Category = Mover)
+	const FMoverDataCollection& GetLocalSimInput() const;
+
+	UFUNCTION(BlueprintPure, Category = Mover)
+	FMoverDataCollection& GetLocalSimInput_Mutable();
 
 	UFUNCTION(BlueprintPure, Category = Mover)
 	const UBaseMovementMode* GetCurrentMovementMode() const;
@@ -102,8 +102,8 @@ protected:
 	void DisableCharacterConstraint();
 
 	FMoverSyncState CurrentSyncState;
-	FVector TargetDeltaPosition;
-	float TargetDeltaFacing;
+	FMoverDataCollection InternalSimData;
+	FMoverDataCollection LocalSimInput;
 
 	UE::ChaosMover::FMoverStateMachine StateMachine;
 
diff --git a/Engine/Plugins/Experimental/ChaosMover/Source/ChaosMover/Public/ChaosMover/ChaosMoverSimulationTypes.h b/Engine/Plugins/Experimental/ChaosMover/Source/ChaosMover/Public/ChaosMover/ChaosMoverSimulationTypes.h
--- a/Engine/Plugins/Experimental/ChaosMover/Source/ChaosMover/Public/ChaosMover/ChaosMoverSimulationTypes.h
+++ b/Engine/Plugins/Experimental/ChaosMover/Source/ChaosMover/Public/ChaosMover/ChaosMoverSimulationTypes.h
@@ -33,8 +33,33 @@ namespace UE::ChaosMover
 
 } // namespace UE::ChaosMover
 
+// Default chaos mover character simulation state, contains state basic for character simulation
+USTRUCT(BlueprintType)
+struct CHAOSMOVER_API FChaosMoverCharacterSimState : public FMoverDataStructBase
+{
+	GENERATED_BODY()
+
+	/** return newly allocated copy of this FMoverDataStructBase. MUST be overridden by derived types. */
+	virtual FMoverDataStructBase* Clone() const override
+	{
+		return new FChaosMoverCharacterSimState(*this);
+	}
+
+	/** Gets the type info of this FMoverDataStructBase. MUST be overridden by derived types. */
+	virtual UScriptStruct* GetScriptStruct() const override
+	{
+		return StaticStruct();
+	}
+
+	UPROPERTY(VisibleAnywhere, BlueprintReadWrite, Category=ChaosMover)
+	FVector TargetDeltaPosition = FVector::ZeroVector;
+
+	UPROPERTY(VisibleAnywhere, BlueprintReadWrite, Category=ChaosMover)
+	float TargetDeltaFacing = 0.0f;
+};
+
 // Data block containing all default inputs required by the Chaos Mover simulation
-USTRUCT()
+USTRUCT(BlueprintType)
 struct CHAOSMOVER_API FChaosMoverSimulationDefaultInputs : public FMoverDataStructBase
 {
 	GENERATED_BODY()
@@ -61,15 +86,18 @@ struct CHAOSMOVER_API FChaosMoverSimulationDefaultInputs : public FMoverDataStru
 	FCollisionResponseParams CollisionResponseParams;
 	FCollisionQueryParams CollisionQueryParams;
 
-	UPROPERTY(BlueprintReadOnly, Category=ChaosMover)
+	UPROPERTY(VisibleAnywhere, BlueprintReadOnly, Category=ChaosMover)
 	FVector UpDir;
-
+	UPROPERTY(VisibleAnywhere, BlueprintReadOnly, Category=ChaosMover)
 	FVector Gravity;
 	Chaos::FPhysicsObjectHandle PhysicsObject;
 	AActor* OwningActor;
 	UWorld* World;
+	UPROPERTY(VisibleAnywhere, BlueprintReadOnly, Category=ChaosMover)
 	float PhysicsObjectGravity;
+	UPROPERTY(VisibleAnywhere, BlueprintReadOnly, Category=ChaosMover)
 	float PawnCollisionHalfHeight;
+	UPROPERTY(VisibleAnywhere, BlueprintReadOnly, Category=ChaosMover)
 	float PawnCollisionRadius;
 	ECollisionChannel CollisionChannel;
 };
