Chaos Mover - Fix flying mode after change to move sim data to a new data collection

diff --git a/Engine/Plugins/Experimental/ChaosMover/Source/ChaosMover/Private/Character/Modes/ChaosFlyingMode.cpp b/Engine/Plugins/Experimental/ChaosMover/Source/ChaosMover/Private/Character/Modes/ChaosFlyingMode.cpp
--- a/Engine/Plugins/Experimental/ChaosMover/Source/ChaosMover/Private/Character/Modes/ChaosFlyingMode.cpp
+++ b/Engine/Plugins/Experimental/ChaosMover/Source/ChaosMover/Private/Character/Modes/ChaosFlyingMode.cpp
@@ -26,17 +26,17 @@ UChaosFlyingMode::UChaosFlyingMode(const FObjectInitializer& ObjectInitializer)
 
 void UChaosFlyingMode::GenerateMove_Implementation(const FMoverTickStartData& StartState, const FMoverTimeStep& TimeStep, FProposedMove& OutProposedMove) const
 {
-	const FChaosMoverSimulationDefaultInputs* DefaultSimInputs = StartState.AuxState.AuxStateCollection.FindDataByType<FChaosMoverSimulationDefaultInputs>();
-	const FCharacterDefaultInputs* CharacterInputs = StartState.InputCmd.InputCollection.FindDataByType<FCharacterDefaultInputs>();
-	if (!DefaultSimInputs || !CharacterInputs)
+	if (!Simulation)
 	{
-		UE_LOG(LogChaosMover, Warning, TEXT("ChaosFallingMode requires FChaosMoverSimulationDefaultInputs and FCharacterDefaultInputs"));
+		UE_LOG(LogChaosMover, Warning, TEXT("No Simulation set on ChaosFlyingMode"));
 		return;
 	}
 
-	if (!Simulation)
+	const FChaosMoverSimulationDefaultInputs* DefaultSimInputs = Simulation->GetLocalSimInput().FindDataByType<FChaosMoverSimulationDefaultInputs>();
+	const FCharacterDefaultInputs* CharacterInputs = StartState.InputCmd.InputCollection.FindDataByType<FCharacterDefaultInputs>();
+	if (!DefaultSimInputs || !CharacterInputs)
 	{
-		UE_LOG(LogChaosMover, Warning, TEXT("No Simulation set on ChaosFallingMode"));
+		UE_LOG(LogChaosMover, Warning, TEXT("ChaosFlyingMode requires FChaosMoverSimulationDefaultInputs and FCharacterDefaultInputs"));
 		return;
 	}
 
@@ -99,10 +99,16 @@ void UChaosFlyingMode::GenerateMove_Implementation(const FMoverTickStartData& St
 
 void UChaosFlyingMode::SimulationTick_Implementation(const FSimulationTickParams& Params, FMoverTickEndData& OutputState)
 {
-	const FChaosMoverSimulationDefaultInputs* DefaultSimInputs = Params.StartState.AuxState.AuxStateCollection.FindDataByType<FChaosMoverSimulationDefaultInputs>();
+	if (!Simulation)
+	{
+		UE_LOG(LogChaosMover, Warning, TEXT("No Simulation set on ChaosFlyingMode"));
+		return;
+	}
+
+	const FChaosMoverSimulationDefaultInputs* DefaultSimInputs = Simulation->GetLocalSimInput().FindDataByType<FChaosMoverSimulationDefaultInputs>();
 	if (!DefaultSimInputs)
 	{
-		UE_LOG(LogChaosMover, Warning, TEXT("ChaosFallingMode requires FChaosMoverSimulationDefaultInputs"));
+		UE_LOG(LogChaosMover, Warning, TEXT("ChaosFlyingMode requires FChaosMoverSimulationDefaultInputs"));
 		return;
 	}
 
