Chaos Mover - Fix jumping by passing through simulation to apply velocity effect

diff --git a/Engine/Plugins/Experimental/ChaosMover/Source/ChaosMover/Private/ChaosMoverSimulation.cpp b/Engine/Plugins/Experimental/ChaosMover/Source/ChaosMover/Private/ChaosMoverSimulation.cpp
--- a/Engine/Plugins/Experimental/ChaosMover/Source/ChaosMover/Private/ChaosMoverSimulation.cpp
+++ b/Engine/Plugins/Experimental/ChaosMover/Source/ChaosMover/Private/ChaosMoverSimulation.cpp
@@ -26,17 +26,6 @@
 
 UChaosMoverSimulation::UChaosMoverSimulation()
 {
-	Blackboard = CreateDefaultSubobject<UMoverBlackboard>(TEXT("ChaosMoverSimulationBlackboard"));
-}
-
-const UMoverBlackboard* UChaosMoverSimulation::GetBlackboard() const
-{
-	return Blackboard;
-}
-
-UMoverBlackboard* UChaosMoverSimulation::GetBlackboard_Mutable()
-{
-	return Blackboard;
 }
 
 const FMoverDataCollection& UChaosMoverSimulation::GetLocalSimInput() const
@@ -98,9 +87,11 @@ void UChaosMoverSimulation::Init(const FInitParams& InitParams)
 
 	CurrentSyncState = InitParams.InitialSyncState;
 
-	StateMachine.InitDefaultModesAndTransitions(InitParams.ImmediateModeTransition, InitParams.NullMovementMode);
-	StateMachine.ClearAllMovementModes();
-	StateMachine.ClearAllGlobalTransitions();
+	UE::ChaosMover::FMoverStateMachine::FInitParams StateMachineInitParams;
+	StateMachineInitParams.ImmediateMovementModeTransition = InitParams.ImmediateModeTransition;
+	StateMachineInitParams.NullMovementMode = InitParams.NullMovementMode;
+	StateMachineInitParams.Simulation = this;
+	StateMachine.Init(StateMachineInitParams);
 
 	for (const TPair<FName, TWeakObjectPtr<UBaseMovementMode>>& Element : InitParams.ModesToRegister)
 	{
diff --git a/Engine/Plugins/Experimental/ChaosMover/Source/ChaosMover/Private/ChaosMoverStateMachine.cpp b/Engine/Plugins/Experimental/ChaosMover/Source/ChaosMover/Private/ChaosMoverStateMachine.cpp
--- a/Engine/Plugins/Experimental/ChaosMover/Source/ChaosMover/Private/ChaosMoverStateMachine.cpp
+++ b/Engine/Plugins/Experimental/ChaosMover/Source/ChaosMover/Private/ChaosMoverStateMachine.cpp
@@ -6,6 +6,7 @@
 #include "Backends/ChaosMoverSubsystem.h"
 #include "ChaosMover/ChaosMoverDeveloperSettings.h"
 #include "ChaosMover/ChaosMoverLog.h"
+#include "ChaosMover/ChaosMoverSimulation.h"
 #include "MoveLibrary/MovementMixer.h"
 #include "MoveLibrary/MovementUtils.h"
 #include "MovementModeTransition.h"
@@ -24,6 +25,16 @@ namespace UE::ChaosMover
 	{
 	}
 
+	void FMoverStateMachine::Init(const FInitParams& Params)
+	{
+		ImmediateMovementModeTransitionWeakPtr = Params.ImmediateMovementModeTransition;
+		NullMovementModeWeakPtr = Params.NullMovementMode;
+		Simulation = Params.Simulation;
+
+		ClearAllMovementModes();
+		ClearAllGlobalTransitions();
+	}
+
 	void FMoverStateMachine::RegisterMovementMode(FName ModeName, TWeakObjectPtr<UBaseMovementMode> Mode, bool bIsDefaultMode)
 	{
 		// JAH TODO: add validation and warnings for overwriting modes
@@ -454,12 +465,6 @@ namespace UE::ChaosMover
 		ClearQueuedMode();
 	}
 
-	void FMoverStateMachine::InitDefaultModesAndTransitions(const TWeakObjectPtr<UImmediateMovementModeTransition> InImmediateMovementModeTransition, const TWeakObjectPtr<UNullMovementMode> InNullMovementMode)
-	{
-		ImmediateMovementModeTransitionWeakPtr = InImmediateMovementModeTransition;
-		NullMovementModeWeakPtr = InNullMovementMode;
-	}
-
 	void FMoverStateMachine::AdvanceToNextMode()
 	{
 		TStrongObjectPtr<UImmediateMovementModeTransition> QueuedModeTransition = QueuedModeTransitionWeakPtr.Pin();
@@ -521,6 +526,7 @@ namespace UE::ChaosMover
 		FApplyMovementEffectParams_Async EffectParams_Async;
 		EffectParams_Async.StartState = &SubstepStartData;
 		EffectParams_Async.TimeStep = &SubTimeStep;
+		EffectParams_Async.Simulation = Simulation.Get();
 
 		bool bInstantMovementEffectApplied = false;
 
diff --git a/Engine/Plugins/Experimental/ChaosMover/Source/ChaosMover/Private/ChaosMoverStateMachine.h b/Engine/Plugins/Experimental/ChaosMover/Source/ChaosMover/Private/ChaosMoverStateMachine.h
--- a/Engine/Plugins/Experimental/ChaosMover/Source/ChaosMover/Private/ChaosMoverStateMachine.h
+++ b/Engine/Plugins/Experimental/ChaosMover/Source/ChaosMover/Private/ChaosMoverStateMachine.h
@@ -8,6 +8,8 @@
 #include "MovementMode.h"
 #include "Templates/SubclassOf.h"
 
+class UChaosMoverSimulation;
+
 namespace UE::ChaosMover
 {
 	/** Mode change event structure, used to postpone callbacks to gameplay code when a mode has changed */
@@ -25,6 +27,14 @@ namespace UE::ChaosMover
 		FMoverStateMachine();
 		virtual ~FMoverStateMachine();
 
+		struct FInitParams
+		{
+			TWeakObjectPtr<UChaosMoverSimulation> Simulation;
+			TWeakObjectPtr<UNullMovementMode> NullMovementMode;
+			TWeakObjectPtr<UImmediateMovementModeTransition> ImmediateMovementModeTransition;
+		};
+		void Init(const FInitParams& Params);
+
 		void RegisterMovementMode(FName ModeName, TWeakObjectPtr<UBaseMovementMode> Mode, bool bIsDefaultMode = false);
 
 		void UnregisterMovementMode(FName ModeName);
@@ -63,8 +73,6 @@ namespace UE::ChaosMover
 		void SetOwnerActorName(const FString& InOwnerActorName);
 		void SetOwnerActorLocalNetRole(ENetRole InOwnerActorLocalNetRole);
 
-		void InitDefaultModesAndTransitions(const TWeakObjectPtr<UImmediateMovementModeTransition> ImmediateMovementModeTransition, const TWeakObjectPtr<UNullMovementMode> NullMovementMode);
-
 	protected:
 
 		TMap<FName, TWeakObjectPtr<UBaseMovementMode>> Modes;
@@ -90,6 +98,8 @@ namespace UE::ChaosMover
 		/** This queue is used to record movement mode transitions and postpone the execution of their callbacks */
 		TArray<FMovementModeChangeEvent> QueuedModeChangeEvents;
 
+		TWeakObjectPtr<UChaosMoverSimulation> Simulation;
+
 	private:
 		void ConstructDefaultModes();
 		void AdvanceToNextMode();
diff --git a/Engine/Plugins/Experimental/ChaosMover/Source/ChaosMover/Private/Character/Effects/ChaosCharacterApplyVelocityEffect.cpp b/Engine/Plugins/Experimental/ChaosMover/Source/ChaosMover/Private/Character/Effects/ChaosCharacterApplyVelocityEffect.cpp
--- a/Engine/Plugins/Experimental/ChaosMover/Source/ChaosMover/Private/Character/Effects/ChaosCharacterApplyVelocityEffect.cpp
+++ b/Engine/Plugins/Experimental/ChaosMover/Source/ChaosMover/Private/Character/Effects/ChaosCharacterApplyVelocityEffect.cpp
@@ -3,16 +3,29 @@
 #include "ChaosMover/Character/Effects/ChaosCharacterApplyVelocityEffect.h"
 
 #include "Chaos/PhysicsObjectInternalInterface.h"
-#include "ChaosMover/ChaosMoverSimulationTypes.h"
+#include "ChaosMover/ChaosMoverLog.h"
+#include "ChaosMover/ChaosMoverSimulation.h"
 
 bool FChaosCharacterApplyVelocityEffect::ApplyMovementEffect_Async(FApplyMovementEffectParams_Async& ApplyEffectParams, FMoverSyncState& OutputState)
 {
 	check(ApplyEffectParams.StartState);
 	const FMoverDefaultSyncState* CurrentSyncState = ApplyEffectParams.StartState->SyncState.SyncStateCollection.FindDataByType<FMoverDefaultSyncState>();
-	const FChaosMoverSimulationDefaultInputs* SimInputs = ApplyEffectParams.StartState->AuxState.AuxStateCollection.FindDataByType<FChaosMoverSimulationDefaultInputs>();
+	if (!CurrentSyncState)
+	{
+		return false;
+	}
+
+	const UChaosMoverSimulation* Simulation = Cast<UChaosMoverSimulation>(ApplyEffectParams.Simulation);
+	if (!Simulation)
+	{
+		UE_LOG(LogChaosMover, Warning, TEXT("No Simulation set on ChaosCharacterApplyVelocityEffect"));
+		return false;
+	}
 
-	if (!CurrentSyncState || !SimInputs)
+	const FChaosMoverSimulationDefaultInputs* SimInputs = Simulation->GetLocalSimInput().FindDataByType<FChaosMoverSimulationDefaultInputs>();
+	if (!SimInputs)
 	{
+		UE_LOG(LogChaosMover, Warning, TEXT("ChaosCharacterApplyVelocityEffect requires FChaosMoverSimulationDefaultInputs"));
 		return false;
 	}
 
diff --git a/Engine/Plugins/Experimental/ChaosMover/Source/ChaosMover/Private/Character/Transitions/ChaosCharacterJumpCheck.cpp b/Engine/Plugins/Experimental/ChaosMover/Source/ChaosMover/Private/Character/Transitions/ChaosCharacterJumpCheck.cpp
--- a/Engine/Plugins/Experimental/ChaosMover/Source/ChaosMover/Private/Character/Transitions/ChaosCharacterJumpCheck.cpp
+++ b/Engine/Plugins/Experimental/ChaosMover/Source/ChaosMover/Private/Character/Transitions/ChaosCharacterJumpCheck.cpp
@@ -47,7 +47,13 @@ FTransitionEvalResult UChaosCharacterJumpCheck::Evaluate_Implementation(const FS
 
 void UChaosCharacterJumpCheck::Trigger_Implementation(const FSimulationTickParams& Params)
 {
-	if (const FChaosMoverSimulationDefaultInputs* SimInputs = Params.StartState.AuxState.AuxStateCollection.FindDataByType<FChaosMoverSimulationDefaultInputs>())
+	if (!Simulation)
+	{
+		UE_LOG(LogChaosMover, Warning, TEXT("No Simulation set on UChaosCharacterJumpCheck"));
+		return;
+	}
+
+	if (const FChaosMoverSimulationDefaultInputs* SimInputs = Simulation->GetLocalSimInput().FindDataByType<FChaosMoverSimulationDefaultInputs>())
 	{
 		FVector JumpVelocity = JumpUpwardsSpeed * SimInputs->UpDir;
 
@@ -55,26 +61,23 @@ void UChaosCharacterJumpCheck::Trigger_Implementation(const FSimulationTickParam
 		JumpMove->VelocityOrImpulseToApply = JumpVelocity;
 		JumpMove->Mode = EChaosMoverVelocityEffectMode::AdditiveVelocity;
 
-		if (ensure(Simulation))
-		{
-			Simulation->QueueInstantMovementEffect(JumpMove);
+		Simulation->QueueInstantMovementEffect(JumpMove);
 
-			// Apply an equal and opposite impulse to the ground
-			if (FractionalGroundReactionImpulse > 0.0f)
+		// Apply an equal and opposite impulse to the ground
+		if (FractionalGroundReactionImpulse > 0.0f)
+		{
+			FFloorCheckResult FloorResult;
+			const UMoverBlackboard* Blackboard = Simulation->GetBlackboard();
+			if (Blackboard && Blackboard->TryGet(CommonBlackboard::LastFloorResult, FloorResult))
 			{
-				FFloorCheckResult FloorResult;
-				const UMoverBlackboard* Blackboard = Simulation->GetBlackboard();
-				if (Blackboard && Blackboard->TryGet(CommonBlackboard::LastFloorResult, FloorResult))
+				if (Chaos::FPhysicsObjectHandle GroundObject = FloorResult.HitResult.PhysicsObject)
 				{
-					if (Chaos::FPhysicsObjectHandle GroundObject = FloorResult.HitResult.PhysicsObject)
+					Chaos::FWritePhysicsObjectInterface_Internal Interface = Chaos::FPhysicsObjectInternalInterface::GetWrite();
+					const float CharacterMass = Interface.GetMass({ SimInputs->PhysicsObject });
+					FVector ImpulseToApply = -FractionalGroundReactionImpulse * CharacterMass * JumpVelocity;
+					if (Chaos::FPBDRigidParticleHandle* GroundParticle = Interface.GetRigidParticle(GroundObject))
 					{
-						Chaos::FWritePhysicsObjectInterface_Internal Interface = Chaos::FPhysicsObjectInternalInterface::GetWrite();
-						const float CharacterMass = Interface.GetMass({ SimInputs->PhysicsObject });
-						FVector ImpulseToApply = -FractionalGroundReactionImpulse * CharacterMass * JumpVelocity;
-						if (Chaos::FPBDRigidParticleHandle* GroundParticle = Interface.GetRigidParticle(GroundObject))
-						{
-							Private::ApplyImpulse(GroundParticle, ImpulseToApply, FloorResult.HitResult.ImpactPoint);
-						}
+						Private::ApplyImpulse(GroundParticle, ImpulseToApply, FloorResult.HitResult.ImpactPoint);
 					}
 				}
 			}
diff --git a/Engine/Plugins/Experimental/ChaosMover/Source/ChaosMover/Public/ChaosMover/ChaosMoverSimulation.h b/Engine/Plugins/Experimental/ChaosMover/Source/ChaosMover/Public/ChaosMover/ChaosMoverSimulation.h
--- a/Engine/Plugins/Experimental/ChaosMover/Source/ChaosMover/Public/ChaosMover/ChaosMoverSimulation.h
+++ b/Engine/Plugins/Experimental/ChaosMover/Source/ChaosMover/Public/ChaosMover/ChaosMoverSimulation.h
@@ -7,6 +7,7 @@
 #include "ChaosMoverStateMachine.h"
 #include "MoveLibrary/MovementMixer.h"
 #include "MoveLibrary/MoverBlackboard.h"
+#include "MoverSimulation.h"
 #include "UObject/WeakObjectPtrTemplates.h"
 
 #include "ChaosMoverSimulation.generated.h"
@@ -23,20 +24,13 @@ namespace UE::ChaosMover
 }
 
 UCLASS(BlueprintType)
-class CHAOSMOVER_API UChaosMoverSimulation : public UObject
+class CHAOSMOVER_API UChaosMoverSimulation : public UMoverSimulation
 {
 	GENERATED_BODY()
 
 public:
-
 	UChaosMoverSimulation();
 
-	UFUNCTION(BlueprintPure, Category = Mover)
-	const UMoverBlackboard* GetBlackboard() const;
-
-	UFUNCTION(BlueprintPure, Category = Mover)
-	UMoverBlackboard* GetBlackboard_Mutable();
-
 	UFUNCTION(BlueprintPure, Category = Mover)
 	const FMoverDataCollection& GetLocalSimInput() const;
 
@@ -107,9 +101,6 @@ protected:
 
 	UE::ChaosMover::FMoverStateMachine StateMachine;
 
-	UPROPERTY(Transient)
-	TObjectPtr<UMoverBlackboard> Blackboard = nullptr;
-
 	TWeakObjectPtr<UMovementMixer> MovementMixerWeakPtr = nullptr;
 	Chaos::FCharacterGroundConstraintProxy* CharacterConstraintProxy = nullptr;
 	Chaos::FConstPhysicsObjectHandle PhysicsObject = nullptr;
