Chaos Mover - Remove jumping check from falling check transition as this prevents jumping for the case where a falling transition check is listed before a jump transition

#rnx

diff --git a/Engine/Plugins/Experimental/ChaosMover/Source/ChaosMover/Private/Character/Transitions/ChaosCharacterFallingCheck.cpp b/Engine/Plugins/Experimental/ChaosMover/Source/ChaosMover/Private/Character/Transitions/ChaosCharacterFallingCheck.cpp
--- a/Engine/Plugins/Experimental/ChaosMover/Source/ChaosMover/Private/Character/Transitions/ChaosCharacterFallingCheck.cpp
+++ b/Engine/Plugins/Experimental/ChaosMover/Source/ChaosMover/Private/Character/Transitions/ChaosCharacterFallingCheck.cpp
@@ -62,9 +62,8 @@ FTransitionEvalResult UChaosCharacterFallingCheck::Evaluate_Implementation(const
 			return EvalResult;
 		}
 
-		const FCharacterDefaultInputs* CharacterInputs = StartState.InputCmd.InputCollection.FindDataByType<FCharacterDefaultInputs>();
 		const FMoverDefaultSyncState* StartingSyncState = StartState.SyncState.SyncStateCollection.FindDataByType<FMoverDefaultSyncState>();
-		if (!CharacterInputs || !StartingSyncState)
+		if (!StartingSyncState)
 		{
 			return EvalResult;
 		}
@@ -88,9 +87,8 @@ FTransitionEvalResult UChaosCharacterFallingCheck::Evaluate_Implementation(const
 			const float ProjectedRelativeNormalVelocity = FloorResult.HitResult.ImpactNormal.Dot(ProjectedVelocity - ProjectedGroundVelocity);
 			const float ProjectedRelativeVerticalVelocity = SimInputs->UpDir.Dot(ProjectedVelocity - ProjectedGroundVelocity);
 			const float VerticalVelocityLimit = 2.0f / DeltaSeconds;
-			const bool bJumping = CharacterInputs && CharacterInputs->bIsJumpJustPressed;
 			bool bIsLiftingOffSurface = false;
-			if (bJumping || ((ProjectedRelativeNormalVelocity > VerticalVelocityLimit) && bIsGroundMoving && (ProjectedRelativeVerticalVelocity > VerticalVelocityLimit)))
+			if ((ProjectedRelativeNormalVelocity > VerticalVelocityLimit) && bIsGroundMoving && (ProjectedRelativeVerticalVelocity > VerticalVelocityLimit))
 			{
 				bIsLiftingOffSurface = true;
 			}
