[ChaosMover] Implementing Swimming Mode

diff --git a/Engine/Plugins/Experimental/ChaosMover/Source/ChaosMover/Private/Character/Modes/ChaosFallingMode.cpp b/Engine/Plugins/Experimental/ChaosMover/Source/ChaosMover/Private/Character/Modes/ChaosFallingMode.cpp
--- a/Engine/Plugins/Experimental/ChaosMover/Source/ChaosMover/Private/Character/Modes/ChaosFallingMode.cpp
+++ b/Engine/Plugins/Experimental/ChaosMover/Source/ChaosMover/Private/Character/Modes/ChaosFallingMode.cpp
@@ -6,6 +6,7 @@
 #include "ChaosMover/ChaosMoverSimulation.h"
 #include "ChaosMover/Character/Settings/SharedChaosCharacterMovementSettings.h"
 #include "ChaosMover/Character/Transitions/ChaosCharacterLandingCheck.h"
+#include "ChaosMover/Character/Transitions/ChaosCharacterWaterCheck.h"
 #include "ChaosMover/Utilities/ChaosMoverQueryUtils.h"
 #include "MoveLibrary/AirMovementUtils.h"
 #include "MoveLibrary/FloorQueryUtils.h"
@@ -36,6 +37,7 @@ UChaosFallingMode::UChaosFallingMode(const FObjectInitializer& ObjectInitializer
 	TwistTorqueLimit = 0.0f;
 
 	Transitions.Add(CreateDefaultSubobject<UChaosCharacterLandingCheck>(TEXT("DefaultLandingCheck")));
+	Transitions.Add(CreateDefaultSubobject<UChaosCharacterWaterCheck>(TEXT("DefaultWaterCheck")));
 }
 
 void UChaosFallingMode::GenerateMove_Implementation(const FMoverTickStartData& StartState, const FMoverTimeStep& TimeStep, FProposedMove& OutProposedMove) const
diff --git a/Engine/Plugins/Experimental/ChaosMover/Source/ChaosMover/Private/Character/Modes/ChaosSwimmingMode.cpp b/Engine/Plugins/Experimental/ChaosMover/Source/ChaosMover/Private/Character/Modes/ChaosSwimmingMode.cpp
new file mode 100644
--- /dev/null
+++ b/Engine/Plugins/Experimental/ChaosMover/Source/ChaosMover/Private/Character/Modes/ChaosSwimmingMode.cpp
@@ -0,0 +1,259 @@
+// Copyright Epic Games, Inc. All Rights Reserved.
+
+#include "ChaosMover/Character/Modes/ChaosSwimmingMode.h"
+
+#include "ChaosMover/ChaosMoverLog.h"
+#include "ChaosMover/ChaosMoverSimulation.h"
+#include "ChaosMover/Character/Settings/SharedChaosCharacterMovementSettings.h"
+#include "ChaosMover/Character/Transitions/ChaosCharacterJumpCheck.h"
+#include "ChaosMover/Character/Transitions/ChaosCharacterWaterCheck.h"
+#include "ChaosMover/Utilities/ChaosMoverQueryUtils.h"
+#include "MoveLibrary/MovementUtils.h"
+#include "MoveLibrary/FloorQueryUtils.h"
+#include "MoveLibrary/WaterMovementUtils.h"
+
+#include UE_INLINE_GENERATED_CPP_BY_NAME(ChaosSwimmingMode)
+
+UChaosSwimmingMode::UChaosSwimmingMode(const FObjectInitializer& ObjectInitializer)
+	: Super(ObjectInitializer)
+{
+	bSupportsAsync = true;
+	GameplayTags.AddTag(Mover_IsSwimming);
+
+	RadialForceLimit = 0.0f;
+	SwingTorqueLimit = 3000.0f;
+	TwistTorqueLimit = 0.0f;
+
+	Transitions.Add(CreateDefaultSubobject<UChaosCharacterWaterCheck>(TEXT("DefaultWaterCheck")));
+	Transitions.Add(CreateDefaultSubobject<UChaosCharacterJumpCheck>(TEXT("DefaultJumpCheck")));
+}
+
+void UChaosSwimmingMode::GenerateMove_Implementation(const FMoverTickStartData& StartState, const FMoverTimeStep& TimeStep, FProposedMove& OutProposedMove) const
+{
+	if (!Simulation)
+	{
+		UE_LOG(LogChaosMover, Warning, TEXT("No Simulation set on ChaosSwimmingMode"));
+		return;
+	}
+
+	const FChaosMoverSimulationDefaultInputs* DefaultSimInputs = Simulation->GetLocalSimInput().FindDataByType<FChaosMoverSimulationDefaultInputs>();
+	const FCharacterDefaultInputs* CharacterInputs = StartState.InputCmd.InputCollection.FindDataByType<FCharacterDefaultInputs>();
+	if (!DefaultSimInputs || !CharacterInputs)
+	{
+		UE_LOG(LogChaosMover, Warning, TEXT("ChaosSwimmingMode requires FChaosMoverSimulationDefaultInputs and FCharacterDefaultInputs"));
+		return;
+	}
+	
+	const FMoverDefaultSyncState* StartingSyncState = StartState.SyncState.SyncStateCollection.FindDataByType<FMoverDefaultSyncState>();
+	check(StartingSyncState);
+	
+	const float DeltaSeconds = TimeStep.StepMs * 0.001f;
+	
+	const float HalfHeight = DefaultSimInputs->PawnCollisionHalfHeight;
+	FVector Velocity = StartingSyncState->GetVelocity_WorldSpace();
+	
+	UMoverBlackboard* SimBlackboard = Simulation->GetBlackboard_Mutable();
+	SimBlackboard->Invalidate(CommonBlackboard::LastFloorResult);
+	SimBlackboard->Invalidate(CommonBlackboard::LastWaterResult);
+	FFloorCheckResult FloorResult;
+	FWaterCheckResult WaterResult;
+
+	const float QueryDistance = 2.0f * HalfHeight;
+	FVector DeltaPos = Velocity * DeltaSeconds;
+	UE::ChaosMover::Utils::FFloorSweepParams SweepParams{
+		.ResponseParams = DefaultSimInputs->CollisionResponseParams,
+		.QueryParams = DefaultSimInputs->CollisionQueryParams,
+		.Location = StartingSyncState->GetLocation_WorldSpace(),
+		.DeltaPos = DeltaPos,
+		.UpDir = DefaultSimInputs->UpDir,
+		.World = DefaultSimInputs->World,
+		.QueryDistance = QueryDistance,
+		.QueryRadius = FMath::Min(GetGroundQueryRadius(), FMath::Max(DefaultSimInputs->PawnCollisionRadius - 5.0f, 0.0f)),
+		.MaxWalkSlopeCosine = GetMaxWalkSlopeCosine(),
+		.TargetHeight = GetTargetHeight(),
+		.CollisionChannel = DefaultSimInputs->CollisionChannel
+	};
+
+	UE::ChaosMover::Utils::FloorSweep_Internal(SweepParams, FloorResult, WaterResult);
+
+	if (WaterResult.IsSwimmableVolume())
+	{
+		FUpdateWaterSplineDataParams SplineParams;
+		SplineParams.TargetImmersionDepth = SwimmingIdealImmersionDepth;
+		SplineParams.WaterVelocityDepthForMax = SurfaceSwimmingWaterControlSettings.WaterVelocityDepthForMax;
+		SplineParams.WaterVelocityMinMultiplier = SurfaceSwimmingWaterControlSettings.WaterVelocityMinMultiplier;
+		SplineParams.PlayerVelocity = StartingSyncState->GetVelocity_WorldSpace();
+		SplineParams.CapsuleHalfHeight = HalfHeight;
+		SplineParams.PlayerLocation = StartingSyncState->GetLocation_WorldSpace();
+
+		UWaterMovementUtils::UpdateWaterSplineData(SplineParams, WaterResult);
+	}
+
+	SimBlackboard->Set(CommonBlackboard::LastFloorResult, FloorResult);
+	SimBlackboard->Set(CommonBlackboard::LastWaterResult, WaterResult);
+	
+	// Buoyancy Bobbing
+	{
+		const FWaterFlowSplineData& WaterData = WaterResult.WaterSplineData;
+		
+		const float ReciprocalHeight = 0.5f / HalfHeight;
+		const float ReciprocalOriginalHeight = 0.5f / OriginalHalfHeight;
+		const FVector Location = StartingSyncState->GetLocation_WorldSpace();
+
+		const float ImmersionDepth = WaterData.ImmersionDepth + HalfHeight;
+		const float ImmersionPercent = FMath::Clamp(ImmersionDepth * ReciprocalHeight, UE_KINDA_SMALL_NUMBER, 1.f);
+		const float IdealDepth = SwimmingIdealImmersionDepth + HalfHeight;
+		const float IdealImmersionPercent = FMath::Clamp(IdealDepth * ReciprocalOriginalHeight, UE_KINDA_SMALL_NUMBER, 1.f);
+
+		const float GravityForce = DefaultSimInputs->Gravity.Z;
+		
+		// Using IdealDepth, we compute the Buoyancy such that (Buoyancy * IdealImmersionPercent) = -G force
+		const float BuoyancyForce = -GravityForce / IdealImmersionPercent;
+	
+		// Now calculate the actual force in the correct ratio
+		float BobbingForce = (BuoyancyForce * ImmersionPercent) + GravityForce;
+		const float MaxForce = SurfaceSwimmingWaterControlSettings.BobbingMaxForce;
+		BobbingForce = FMath::Clamp(BobbingForce, -MaxForce, MaxForce);
+
+		Velocity.Z += BobbingForce * DeltaSeconds;
+
+		// Vertical fluid friction for bobbing
+		if (!FMath::IsNearlyZero(Velocity.Z, (FVector::FReal)0.1f))
+		{
+			const float IdealDepthTolerance = SurfaceSwimmingWaterControlSettings.BobbingIdealDepthTolerance;
+			if (FMath::Sign(Velocity.Z) != FMath::Sign(BobbingForce) || FMath::IsNearlyEqual(ImmersionDepth, IdealDepth, IdealDepthTolerance))
+			{
+				float FluidFriction = 0.f;
+				float ExpDrag = 0.f;
+				const bool bFullySubmerged = ImmersionPercent > 1.f;
+				if (Velocity.Z > 0)
+				{
+					FluidFriction = SurfaceSwimmingWaterControlSettings.BobbingFrictionUp;
+					ExpDrag = SurfaceSwimmingWaterControlSettings.BobbingExpDragUp;
+				}
+				else
+				{
+					// Different drag when fully immersed and moving down (mainly controls how far you go when falling in fast)
+					FluidFriction = bFullySubmerged ? SurfaceSwimmingWaterControlSettings.BobbingFrictionDownSubmerged : SurfaceSwimmingWaterControlSettings.BobbingFrictionDown;
+					ExpDrag = bFullySubmerged ? SurfaceSwimmingWaterControlSettings.BobbingExpDragDownSubmerged : SurfaceSwimmingWaterControlSettings.BobbingExpDragDown;
+				}
+	
+				FluidFriction *= SurfaceSwimmingWaterControlSettings.BobbingFrictionMultiplier;
+				ExpDrag *= SurfaceSwimmingWaterControlSettings.BobbingExpDragMultiplier;
+	
+				Velocity.Z = Velocity.Z * (1.f - FMath::Min(FluidFriction * DeltaSeconds, 1.f));
+				Velocity.Z = Velocity.Z * (1.f - FMath::Min(FMath::Abs(Velocity.Z) * FMath::Square(ExpDrag) * DeltaSeconds, 1.f));
+			}
+		}
+	}
+
+	// Vertical speed limit in Water
+	{
+		const float MaxVerticalWaterSpeedUp = SurfaceSwimmingWaterControlSettings.MaxSpeedUp;
+		const float MaxVerticalWaterSpeedDown = SurfaceSwimmingWaterControlSettings.MaxSpeedDown;
+	
+		Velocity.Z = FMath::Clamp(Velocity.Z, -FMath::Abs(MaxVerticalWaterSpeedDown), FMath::Abs(MaxVerticalWaterSpeedUp));
+	}
+
+	// Calculate and apply the requested move here
+	{
+		// Force from Water flow's WaterVelocity 
+		const float MaxWaterForce = SurfaceSwimmingWaterControlSettings.MaxWaterForce;
+		const float WaterForceMultiplier = SurfaceSwimmingWaterControlSettings.WaterForceMultiplier * SurfaceSwimmingWaterControlSettings.WaterForceSecondMultiplier;
+		const FVector WaterVelocity = WaterResult.WaterSplineData.WaterVelocity;
+		const FVector WaterAcceleration = (WaterVelocity * WaterForceMultiplier).GetClampedToMaxSize(MaxWaterForce);
+		const float WaterSpeed = WaterVelocity.Size();
+	
+	    // Consider player input
+		FRotator IntendedOrientation_WorldSpace;
+		if (!CharacterInputs || CharacterInputs->OrientationIntent.IsNearlyZero())
+		{
+			IntendedOrientation_WorldSpace = StartingSyncState->GetOrientation_WorldSpace();
+		}
+		else
+		{
+			IntendedOrientation_WorldSpace = CharacterInputs->GetOrientationIntentDir_WorldSpace().ToOrientationRotator();
+		}
+
+		const FQuat WorldToGravityTransform = FQuat::FindBetweenNormals(FVector::UpVector, DefaultSimInputs->UpDir);
+		IntendedOrientation_WorldSpace = UMovementUtils::ApplyGravityToOrientationIntent(IntendedOrientation_WorldSpace, WorldToGravityTransform, bShouldCharacterRemainUpright);
+		
+		FWaterMoveParams Params;
+		if (CharacterInputs)
+		{
+			Params.MoveInputType = CharacterInputs->GetMoveInputType();
+
+			const bool bMaintainInputMagnitude = true;
+			Params.MoveInput = CharacterInputs->GetMoveInput_WorldSpace();
+			// Params.MoveInput = UPlanarConstraintUtils::ConstrainDirectionToPlane(MoverComp->GetPlanarConstraint(), CharacterInputs->GetMoveInput_WorldSpace(), bMaintainInputMagnitude);
+		}
+		else
+		{
+			Params.MoveInputType = EMoveInputType::None;
+			Params.MoveInput = FVector::ZeroVector;
+		}
+		
+		Params.OrientationIntent = IntendedOrientation_WorldSpace;
+		Params.PriorVelocity = StartingSyncState->GetVelocity_WorldSpace();
+		Params.PriorOrientation = StartingSyncState->GetOrientation_WorldSpace();
+		Params.TurningRate = SharedSettings->TurningRate;
+		Params.TurningBoost = SharedSettings->TurningBoost;
+		Params.MaxSpeed = SharedSettings->MaxSpeed;
+		Params.Acceleration = SharedSettings->Acceleration;
+		Params.Deceleration = SharedSettings->Deceleration;
+		Params.DeltaSeconds = DeltaSeconds;
+		Params.MoveSpeed = WaterSpeed;
+		Params.MoveAcceleration = WaterAcceleration;
+		Params.WorldToGravityQuat = WorldToGravityTransform;
+
+		// Calculate the move
+		OutProposedMove = UWaterMovementUtils::ComputeControlledWaterMove(Params);
+		
+		// Use Z Velocity calculated earlier (Buoyancy, Friction and Terminal Velocity) for the move's Z component
+		OutProposedMove.LinearVelocity.Z = Velocity.Z;
+	}
+}
+
+void UChaosSwimmingMode::SimulationTick_Implementation(const FSimulationTickParams& Params, FMoverTickEndData& OutputState)
+{
+	if (!Simulation)
+	{
+		UE_LOG(LogChaosMover, Warning, TEXT("No Simulation set on ChaosSwimmingMode"));
+		return;
+	}
+	
+	const FChaosMoverSimulationDefaultInputs* DefaultSimInputs = Simulation->GetLocalSimInput().FindDataByType<FChaosMoverSimulationDefaultInputs>();
+	const FMoverTickStartData& StartState = Params.StartState;
+	const FCharacterDefaultInputs* CharacterInputs = StartState.InputCmd.InputCollection.FindDataByType<FCharacterDefaultInputs>();
+	FProposedMove ProposedMove = Params.ProposedMove;
+	const FVector UpDir = DefaultSimInputs->UpDir;
+
+	const FMoverDefaultSyncState* StartingSyncState = StartState.SyncState.SyncStateCollection.FindDataByType<FMoverDefaultSyncState>();
+	check(StartingSyncState);
+
+	FMoverDefaultSyncState& OutputSyncState = OutputState.SyncState.SyncStateCollection.FindOrAddMutableDataByType<FMoverDefaultSyncState>();
+	OutputSyncState = *StartingSyncState;
+
+	OutputSyncState.MoveDirectionIntent = ProposedMove.bHasDirIntent ? ProposedMove.DirectionIntent : FVector::ZeroVector;
+
+	const float DeltaSeconds = Params.TimeStep.StepMs * 0.001f;
+
+	FRotator TargetOrient = StartingSyncState->GetOrientation_WorldSpace();
+	if (!UMovementUtils::IsAngularVelocityZero(ProposedMove.AngularVelocity))
+	{
+		TargetOrient += (ProposedMove.AngularVelocity * DeltaSeconds);
+	}
+
+	FVector TargetVel = ProposedMove.LinearVelocity;
+	TargetVel -= DefaultSimInputs->PhysicsObjectGravity * DefaultSimInputs->UpDir * DeltaSeconds;;
+
+	FVector TargetPos = StartingSyncState->GetLocation_WorldSpace();
+	TargetPos += TargetVel * DeltaSeconds;
+
+	OutputSyncState.SetTransforms_WorldSpace(
+		TargetPos,
+		TargetOrient,
+		TargetVel);
+
+	OutputState.MovementEndState.RemainingMs = 0.0f;
+}
\ No newline at end of file
diff --git a/Engine/Plugins/Experimental/ChaosMover/Source/ChaosMover/Private/Character/Modes/ChaosWalkingMode.cpp b/Engine/Plugins/Experimental/ChaosMover/Source/ChaosMover/Private/Character/Modes/ChaosWalkingMode.cpp
--- a/Engine/Plugins/Experimental/ChaosMover/Source/ChaosMover/Private/Character/Modes/ChaosWalkingMode.cpp
+++ b/Engine/Plugins/Experimental/ChaosMover/Source/ChaosMover/Private/Character/Modes/ChaosWalkingMode.cpp
@@ -11,6 +11,7 @@
 #include "ChaosMover/Character/Settings/SharedChaosCharacterMovementSettings.h"
 #include "ChaosMover/Character/Transitions/ChaosCharacterFallingCheck.h"
 #include "ChaosMover/Character/Transitions/ChaosCharacterJumpCheck.h"
+#include "ChaosMover/Character/Transitions/ChaosCharacterWaterCheck.h"
 #include "ChaosMover/Utilities/ChaosGroundMovementUtils.h"
 #include "ChaosMover/Utilities/ChaosMoverQueryUtils.h"
 #include "Math/UnitConversion.h"
@@ -36,6 +37,7 @@ UChaosWalkingMode::UChaosWalkingMode(const FObjectInitializer& ObjectInitializer
 
 	Transitions.Add(CreateDefaultSubobject<UChaosCharacterFallingCheck>(TEXT("DefaultFallingCheck")));
 	Transitions.Add(CreateDefaultSubobject<UChaosCharacterJumpCheck>(TEXT("DefaultJumpCheck")));
+	Transitions.Add(CreateDefaultSubobject<UChaosCharacterWaterCheck>(TEXT("DefaultWaterCheck")));
 }
 
 void UChaosWalkingMode::UpdateConstraintSettings(Chaos::FCharacterGroundConstraintSettings& ConstraintSettings) const
diff --git a/Engine/Plugins/Experimental/ChaosMover/Source/ChaosMover/Private/Character/Transitions/ChaosCharacterLandingCheck.cpp b/Engine/Plugins/Experimental/ChaosMover/Source/ChaosMover/Private/Character/Transitions/ChaosCharacterLandingCheck.cpp
--- a/Engine/Plugins/Experimental/ChaosMover/Source/ChaosMover/Private/Character/Transitions/ChaosCharacterLandingCheck.cpp
+++ b/Engine/Plugins/Experimental/ChaosMover/Source/ChaosMover/Private/Character/Transitions/ChaosCharacterLandingCheck.cpp
@@ -18,7 +18,6 @@ UChaosCharacterLandingCheck::UChaosCharacterLandingCheck(const FObjectInitialize
 	bSupportsAsync = true;
 
 	TransitionToGroundMode = DefaultModeNames::Walking;
-	TransitionToWaterMode = NAME_None;
 }
 
 FTransitionEvalResult UChaosCharacterLandingCheck::Evaluate_Implementation(const FSimulationTickParams& Params) const
@@ -49,31 +48,6 @@ FTransitionEvalResult UChaosCharacterLandingCheck::Evaluate_Implementation(const
 			const bool bJumping = CharacterInputs && CharacterInputs->bIsJumpJustPressed;
 			const bool bIsMovingUp = bJumping || (VerticalVelocity > 0.0f);
 
-			// Check for water landing
-			if (!TransitionToWaterMode.IsNone())
-			{
-				FWaterCheckResult WaterResult;
-				if (Blackboard->TryGet(CommonBlackboard::LastFloorResult, WaterResult) && WaterResult.IsSwimmableVolume())
-				{
-					const float ProjectedImmersionDepth = WaterResult.WaterSplineData.ImmersionDepth - VerticalVelocity * DeltaSeconds;
-					const bool bStartSwimming = ProjectedImmersionDepth > SwimmingStartImmersionDepth;
-
-					if (bStartSwimming && !bIsMovingUp)
-					{
-						if (Simulation->FindMovementModeByName(TransitionToWaterMode))
-						{
-							EvalResult.NextMode = TransitionToWaterMode;
-						}
-						else
-						{
-							UE_LOG(LogChaosMover, Warning, TEXT("Invalid water mode name %s in ChaosCharacterLandingCheck. Cannot make transition"), *TransitionToWaterMode.ToString());
-						}
-
-						return EvalResult;
-					}
-				}
-			}
-
 			// Check for ground landing
 			if (!TransitionToGroundMode.IsNone())
 			{
diff --git a/Engine/Plugins/Experimental/ChaosMover/Source/ChaosMover/Private/Character/Transitions/ChaosCharacterWaterCheck.cpp b/Engine/Plugins/Experimental/ChaosMover/Source/ChaosMover/Private/Character/Transitions/ChaosCharacterWaterCheck.cpp
new file mode 100644
--- /dev/null
+++ b/Engine/Plugins/Experimental/ChaosMover/Source/ChaosMover/Private/Character/Transitions/ChaosCharacterWaterCheck.cpp
@@ -0,0 +1,121 @@
+// Copyright Epic Games, Inc. All Rights Reserved.
+
+#include "ChaosMover/Character/Transitions/ChaosCharacterWaterCheck.h"
+
+#include "CharacterMovementComponentAsync.h"
+#include "ChaosMover/ChaosMoverLog.h"
+#include "ChaosMover/ChaosMoverSimulation.h"
+#include "ChaosMover/ChaosMoverSimulationTypes.h"
+#include "MoveLibrary/WaterMovementUtils.h"
+#include "MoverSimulationTypes.h"
+
+#include UE_INLINE_GENERATED_CPP_BY_NAME(ChaosCharacterWaterCheck)
+
+UChaosCharacterWaterCheck::UChaosCharacterWaterCheck(const FObjectInitializer& ObjectInitializer)
+	: Super(ObjectInitializer)
+{
+	bSupportsAsync = true;
+
+	WaterModeName = DefaultModeNames::Swimming;
+	GroundModeName = DefaultModeNames::Walking;
+	AirModeName = DefaultModeNames::Falling;
+}
+
+FTransitionEvalResult UChaosCharacterWaterCheck::Evaluate_Implementation(const FSimulationTickParams& Params) const
+{
+	check(Simulation);
+
+	FTransitionEvalResult EvalResult;
+
+	const FMoverTickStartData& StartState = Params.StartState;
+
+	if (const FChaosMoverSimulationDefaultInputs* SimInputs = Simulation->GetLocalSimInput().FindDataByType<FChaosMoverSimulationDefaultInputs>())
+	{
+		const UMoverBlackboard* Blackboard = Simulation->GetBlackboard();
+		check(Blackboard);
+
+		const FCharacterDefaultInputs* CharacterInputs = StartState.InputCmd.InputCollection.FindDataByType<FCharacterDefaultInputs>();
+		const FMoverDefaultSyncState* StartingSyncState = StartState.SyncState.SyncStateCollection.FindDataByType<FMoverDefaultSyncState>();
+		if (!CharacterInputs || !StartingSyncState)
+		{
+			return EvalResult;
+		}
+
+		if (const IChaosCharacterMovementModeInterface* Mode = Cast<IChaosCharacterMovementModeInterface>(Simulation->FindMovementModeByName(Params.StartState.SyncState.MovementMode)))
+		{
+			const float DeltaSeconds = Params.TimeStep.StepMs * 0.001f;
+			const FVector LinearVelocity = StartingSyncState->GetVelocity_WorldSpace();
+			const float VerticalVelocity = SimInputs->UpDir.Dot(LinearVelocity);
+			const bool bJumping = CharacterInputs && CharacterInputs->bIsJumpJustPressed;
+			const bool bIsMovingUp = bJumping || (VerticalVelocity > 0.0f);
+
+			FFindFloorResult FloorResult;
+			FWaterCheckResult WaterResult;
+
+			Blackboard->TryGet(CommonBlackboard::LastFloorResult, FloorResult);
+			Blackboard->TryGet(CommonBlackboard::LastWaterResult, WaterResult);
+			
+			const float ProjectedImmersionDepth = WaterResult.WaterSplineData.ImmersionDepth - VerticalVelocity * DeltaSeconds;
+			const bool bStartSwimming = ProjectedImmersionDepth > WaterModeStartImmersionDepth;
+			const bool bStopSwimming = ProjectedImmersionDepth < WaterModeStopImmersionDepth;
+
+			// Check if we need to enter/exit water mode
+			if (WaterResult.IsSwimmableVolume() && bStartSwimming)
+			{
+				if (!WaterModeName.IsNone())
+				{
+					if (Simulation->FindMovementModeByName(WaterModeName))
+					{
+						EvalResult.NextMode = WaterModeName;
+					}
+					else
+					{
+						UE_LOG(LogChaosMover, Warning, TEXT("Invalid water mode name %s in ChaosCharacterWaterCheck. Cannot make transition"), *WaterModeName.ToString());
+					}
+					return EvalResult;
+				}
+			}
+			else if (bStopSwimming)
+			{
+				const bool bIsWithinReach = FloorResult.FloorDist <= Mode->GetTargetHeight();
+
+				if (FloorResult.IsWalkableFloor() && bIsWithinReach)
+				{
+					if (!GroundModeName.IsNone())
+					{
+						if (Simulation->FindMovementModeByName(GroundModeName))
+						{
+							EvalResult.NextMode = GroundModeName;
+						}
+						else
+						{
+							UE_LOG(LogChaosMover, Warning, TEXT("Invalid ground mode name %s in ChaosCharacterWaterCheck. Cannot make transition"), *GroundModeName.ToString());
+						}
+						return EvalResult;
+					}
+				}
+				else
+				{
+					if (!AirModeName.IsNone())
+					{
+						if (Simulation->FindMovementModeByName(AirModeName))
+						{
+							EvalResult.NextMode = AirModeName;
+						}
+						else
+						{
+							UE_LOG(LogChaosMover, Warning, TEXT("Invalid air mode name %s in ChaosCharacterWaterCheck. Cannot make transition"), *AirModeName.ToString());
+						}
+						return EvalResult;
+					}
+				}
+			}
+		}
+	}
+
+	return EvalResult;
+}
+
+void UChaosCharacterWaterCheck::Trigger_Implementation(const FSimulationTickParams& Params)
+{
+}
diff --git a/Engine/Plugins/Experimental/ChaosMover/Source/ChaosMover/Public/ChaosMover/Character/Modes/ChaosSwimmingMode.h b/Engine/Plugins/Experimental/ChaosMover/Source/ChaosMover/Public/ChaosMover/Character/Modes/ChaosSwimmingMode.h
new file mode 100644
--- /dev/null
+++ b/Engine/Plugins/Experimental/ChaosMover/Source/ChaosMover/Public/ChaosMover/Character/Modes/ChaosSwimmingMode.h
@@ -0,0 +1,107 @@
+// Copyright Epic Games, Inc. All Rights Reserved.
+
+#pragma once
+
+#include "ChaosMover/Character/Modes/ChaosCharacterMovementMode.h"
+
+#include "ChaosSwimmingMode.generated.h"
+
+// Controls for the Swimming Movement
+USTRUCT(BlueprintType)
+struct FSwimmingSettings
+{
+	GENERATED_BODY()
+
+	// Max speed when moving up in water
+	UPROPERTY(EditAnywhere, BlueprintReadWrite, Category = SwimmingControl)
+	float MaxSpeedUp = 500.00f;
+
+	// Max speed when moving down in water
+	UPROPERTY(EditAnywhere, BlueprintReadWrite, Category = SwimmingControl)
+	float MaxSpeedDown = 1000.00f;
+
+	// At or above this depth, use max velocity. Interps down to WaterVelocityMinMultiplier at wading depth (where player can start swimming)
+	UPROPERTY(EditAnywhere, BlueprintReadWrite, Category = SwimmingControl)
+	float WaterVelocityDepthForMax = 175.00f;
+
+	// Min velocity multiplier applied when depth equals min swimming depth (where they transition from wading to swimming). Interps between this and 1.0 at WaterVelocityDepthForMax.
+	UPROPERTY(EditAnywhere, BlueprintReadWrite, Category = SwimmingControl)
+	float WaterVelocityMinMultiplier = 0.50f;
+
+	// Max water force, after WaterVelocity * (WaterForceMultiplier * WaterForceSecondMultiplier).
+	UPROPERTY(EditAnywhere, BlueprintReadWrite, Category = SwimmingControl)
+	float MaxWaterForce = 400.00f;
+
+	// Multiplier to water force acceleration in direction of current.
+	UPROPERTY(EditAnywhere, BlueprintReadWrite, Category = SwimmingControl)
+	float WaterForceMultiplier = 2.00f;
+
+	// Multiplier applied on the top of WaterForceMultiplier, to water force acceleration in direction of current. Used only for inherited objects.
+	UPROPERTY(EditAnywhere, BlueprintReadWrite, Category = SwimmingControl)
+	float WaterForceSecondMultiplier = 1.00f;
+
+	// Bobbing: Max force
+	UPROPERTY(EditAnywhere, BlueprintReadWrite, Category = SwimmingControl)
+	float BobbingMaxForce = 3800.00f;
+
+	// Bobbing: Slow down strongly when within this tolerance of the ideal immersion depth. Normally we apply drag only when going away from the ideal depth, this allows some slowdown when approaching it.
+	UPROPERTY(EditAnywhere, BlueprintReadWrite, Category = SwimmingControl)
+	float BobbingIdealDepthTolerance = 7.50f;
+
+	// Bobbing: friction/drag opposed to downward velocity, linear multiplier per second
+	UPROPERTY(EditAnywhere, BlueprintReadWrite, Category = SwimmingControl)
+	float BobbingFrictionDown = 5.00f;
+
+	// Bobbing: friction/drag opposed to downward velocity, squared with velocity per second. Ramps up faster with higher speeds, less effect at low speeds.
+	UPROPERTY(EditAnywhere, BlueprintReadWrite, Category = SwimmingControl)
+	float BobbingExpDragDown = 0.05f;
+
+	// Bobbing: friction/drag opposed to downward velocity, linear multiplier per second. Only used when fully submerged (replaces other value).
+	UPROPERTY(EditAnywhere, BlueprintReadWrite, Category = SwimmingControl)
+	float BobbingFrictionDownSubmerged = 7.50f;
+
+	// Bobbing: friction/drag opposed to upward velocity, squared with velocity per second. Ramps up faster with higher speeds, less effect at low speeds.
+	UPROPERTY(EditAnywhere, BlueprintReadWrite, Category = SwimmingControl)
+	float BobbingExpDragDownSubmerged = 0.10f;
+
+	// Bobbing: friction/drag opposed to upward velocity, linear multiplier per second
+	UPROPERTY(EditAnywhere, BlueprintReadWrite, Category = SwimmingControl)
+	float BobbingFrictionUp = 7.00f;
+
+	// Bobbing: friction/drag opposed to upward velocity, squared with velocity per second. Ramps up faster with higher speeds, less effect at low speeds.
+	UPROPERTY(EditAnywhere, BlueprintReadWrite, Category = SwimmingControl)
+	float BobbingExpDragUp = 0.20f;
+
+	// Bobbing: friction multiplier, multiplies the fluid friction value.
+	UPROPERTY(EditAnywhere, BlueprintReadWrite, Category = SwimmingControl)
+	float BobbingFrictionMultiplier = 1.f;
+
+	// Bobbing: multiplier for the exponential drag.
+	UPROPERTY(EditAnywhere, BlueprintReadWrite, Category = SwimmingControl)
+	float BobbingExpDragMultiplier = 1.f;
+};
+
+/**
+ * Chaos character swimming mode
+ */
+UCLASS(Blueprintable, BlueprintType, EditInlineNew, DefaultToInstanced)
+class CHAOSMOVER_API UChaosSwimmingMode : public UChaosCharacterMovementMode
+{
+	GENERATED_BODY()
+
+public:
+	UChaosSwimmingMode(const FObjectInitializer& ObjectInitializer);
+
+	// Ideal depth when in water. Measured from the center of the collision shape.
+	UPROPERTY(EditAnywhere, BlueprintReadOnly, Category = "Water Control")
+	float SwimmingIdealImmersionDepth = 40.f;
+
+	UPROPERTY(EditAnywhere, BlueprintReadOnly, Category = "Water Control")
+	FSwimmingSettings SurfaceSwimmingWaterControlSettings;
+	
+	virtual void GenerateMove_Implementation(const FMoverTickStartData& StartState, const FMoverTimeStep& TimeStep, FProposedMove& OutProposedMove) const override;
+	virtual void SimulationTick_Implementation(const FSimulationTickParams& Params, FMoverTickEndData& OutputState) override;
+
+protected:
+	float OriginalHalfHeight = 75.f;
+};
\ No newline at end of file
diff --git a/Engine/Plugins/Experimental/ChaosMover/Source/ChaosMover/Public/ChaosMover/Character/Transitions/ChaosCharacterLandingCheck.h b/Engine/Plugins/Experimental/ChaosMover/Source/ChaosMover/Public/ChaosMover/Character/Transitions/ChaosCharacterLandingCheck.h
--- a/Engine/Plugins/Experimental/ChaosMover/Source/ChaosMover/Public/ChaosMover/Character/Transitions/ChaosCharacterLandingCheck.h
+++ b/Engine/Plugins/Experimental/ChaosMover/Source/ChaosMover/Public/ChaosMover/Character/Transitions/ChaosCharacterLandingCheck.h
@@ -29,8 +29,4 @@ public:
 	/** Name of movement mode to transition to when landing on ground */
 	UPROPERTY(BlueprintReadWrite, EditAnywhere, Category = Parameters)
 	FName TransitionToGroundMode;
-
-	/** Name of movement mode to transition to when landing on water */
-	UPROPERTY(BlueprintReadWrite, EditAnywhere, Category = Parameters)
-	FName TransitionToWaterMode;
 };
\ No newline at end of file
diff --git a/Engine/Plugins/Experimental/ChaosMover/Source/ChaosMover/Public/ChaosMover/Character/Transitions/ChaosCharacterWaterCheck.h b/Engine/Plugins/Experimental/ChaosMover/Source/ChaosMover/Public/ChaosMover/Character/Transitions/ChaosCharacterWaterCheck.h
new file mode 100644
--- /dev/null
+++ b/Engine/Plugins/Experimental/ChaosMover/Source/ChaosMover/Public/ChaosMover/Character/Transitions/ChaosCharacterWaterCheck.h
@@ -0,0 +1,40 @@
+// Copyright Epic Games, Inc. All Rights Reserved.
+
+#pragma once
+
+#include "ChaosMover/ChaosMovementModeTransition.h"
+
+#include "ChaosCharacterWaterCheck.generated.h"
+
+
+UCLASS(Blueprintable, BlueprintType, EditInlineNew, DefaultToInstanced)
+class CHAOSMOVER_API UChaosCharacterWaterCheck : public UChaosMovementModeTransition
+{
+	GENERATED_BODY()
+
+public:
+	UChaosCharacterWaterCheck(const FObjectInitializer& ObjectInitializer);
+
+	virtual FTransitionEvalResult Evaluate_Implementation(const FSimulationTickParams& Params) const override;
+	virtual void Trigger_Implementation(const FSimulationTickParams& Params) override;
+
+	// Depth at which the pawn starts water mode. Measured from the center of the collision shape.
+	UPROPERTY(BlueprintReadWrite, EditAnywhere, Category = Parameters, meta = (Units = "cm"))
+	float WaterModeStartImmersionDepth = 45.f;
+
+	// Depth at which the pawn stops water mode. Measured from the center of the collision shape.
+	UPROPERTY(BlueprintReadWrite, EditAnywhere, Category = Parameters, meta = (Units = "cm"))
+	float WaterModeStopImmersionDepth = 40.f;
+
+	// Name of movement mode to transition to when immersed in water.
+	UPROPERTY(BlueprintReadWrite, EditAnywhere, Category = Parameters)
+	FName WaterModeName;
+
+	// Name of movement mode to transition to when ground is within reach.
+	UPROPERTY(BlueprintReadWrite, EditAnywhere, Category = Parameters)
+	FName GroundModeName;
+
+	// Name of movement mode to transition to when exiting water but ground is not in reach.
+	UPROPERTY(BlueprintReadWrite, EditAnywhere, Category = Parameters)
+	FName AirModeName;
+};
\ No newline at end of file
