Unshelved from pending changelist '41663965':

Chaos Mover: Support for pathed movement modes, kinematic or constraint based.
- IChaosPathedMovementModeInterface. When the current mode implements this interface it indicates that the controlled particle should target the transform set in the output's default sync state. If in kinematic mode, the particle will be set to kinematic and will target that transform kinematically. If in constraint mode, a joint constraint is used, which target end point will be moved kinematically to that target transform.
- Added UChaosSimulation::OnMovementModeChange function so it can react to the mode changing. This allows us to move init code that was repeated every frame to being run once there. In the case of a pathed movement mode, the particle is set to kinematic or dynamic, depending on the objects state desired by the mode, and enables the constraint if in constraint mode.
- Example of a (very) simple use of the pathed movement mode in MoverExamples: ChaosFollowPathMode. This path causes the update component to follow a set of points over a duration. This is a conversion of FollowPathMode that Justin Hare developped previously.
- Removed FMoverStateMachine::SetModeImmediately from UChaosSimulation::Init (GT call) because it will cause OnMovementModeChanged to run, which can modify physics. The initial SyncState must have the starting movement mode which will cause the first update of the state machine to transition from null mode to that one, so we don't need to manually call SetModeImmediately
- Moved the mode change event queue to the simulation and put a mutex around its use to prevent crashes. Tom will rework the implementation as it is still not good enough.
- Moved the tracing to CVD to the simulation post sim

Mover: Fix for a crash I encountered, when a required type previously defined is deleted from code, leaving a null entry. Even if we had a IsDataValid function checking for that, it wouldn't catch this case

Chaos Visual Debugger: Added the ability to pass more than one MoverDataCollection to CVD, as the sim has several. Eventually CVD would display which types are from which collection, when we start passing the information as FInstancedPropertyBags.

diff --git a/Engine/Plugins/Experimental/ChaosMover/Source/ChaosMover/Private/Backends/ChaosMoverBackend.cpp b/Engine/Plugins/Experimental/ChaosMover/Source/ChaosMover/Private/Backends/ChaosMoverBackend.cpp
--- a/Engine/Plugins/Experimental/ChaosMover/Source/ChaosMover/Private/Backends/ChaosMoverBackend.cpp
+++ b/Engine/Plugins/Experimental/ChaosMover/Source/ChaosMover/Private/Backends/ChaosMoverBackend.cpp
@@ -14,19 +14,17 @@
 #include "MovementModeStateMachine.h"
 #include "NetworkChaosMoverData.h"
 #include "PBDRigidsSolver.h"
-#include "Physics/Experimental/PhysScene_Chaos.h"
 #include "Physics/NetworkPhysicsComponent.h"
 #include "PhysicsEngine/PhysicsObjectExternalInterface.h"
 #include "PhysicsProxy/CharacterGroundConstraintProxy.h"
+#include "PhysicsProxy/JointConstraintProxy.h"
+#include "Runtime/Experimental/Chaos/Private/Chaos/PhysicsObjectInternal.h"
 
-#if WITH_CHAOS_VISUAL_DEBUGGER
-#include "ChaosVisualDebugger/ChaosVisualDebuggerTrace.h"
-#include "ChaosVisualDebugger/MoverCVDRuntimeTrace.h"
-#endif
 
 #include UE_INLINE_GENERATED_CPP_BY_NAME(ChaosMoverBackend)
 
 UChaosMoverBackendComponent::UChaosMoverBackendComponent()
+	: PathTargetConstraintPhysicsUserData(&PathTargetConstraintInstance)
 {
 	PrimaryComponentTick.bCanEverTick = false;
 
@@ -56,35 +54,46 @@ void UChaosMoverBackendComponent::InitializeComponent()
 		ImmediateModeTransition = NewObject<UImmediateMovementModeTransition>(&GetMoverComponent(), TEXT("ImmediateModeTransition"));
 
 		// Create NetworkPhysicsComponent
-		if (Chaos::FPhysicsSolverBase::IsNetworkPhysicsPredictionEnabled() && bIsUsingAsyncPhysics)
+		if ((World->GetNetMode() != ENetMode::NM_Standalone) && Chaos::FPhysicsSolverBase::IsNetworkPhysicsPredictionEnabled())
 		{
-			NetworkPhysicsComponent = NewObject<UNetworkPhysicsComponent>(GetOwner(), TEXT("PhysMover_NetworkPhysicsComponent"));
-
-			// This isn't technically a DSO component, but set it net addressable as though it is
-			NetworkPhysicsComponent->SetNetAddressable();
-			NetworkPhysicsComponent->SetIsReplicated(true);
-			NetworkPhysicsComponent->RegisterComponent();
-			if (!NetworkPhysicsComponent->HasBeenInitialized())
+			if (!bIsUsingAsyncPhysics)
 			{
-				NetworkPhysicsComponent->InitializeComponent();
+				// Verify that the Project Settings have bTickPhysicsAsync turned on.
+				// It's easy to waste time forgetting that, since they are off by default.
+				UE_LOG(LogChaosMover, Warning, TEXT("Chaos Mover Backend only supports networking with Physics Async. Networked Physics will not work well. Turn on 'Project Settings > Engine - Physics > Tick Physics Async', or play in Standalone Mode"));
+				// This is important enough that we break for developers debugging in editor
+				UE_DEBUG_BREAK();
 			}
-			NetworkPhysicsComponent->Activate(true);
-
-			// Register network data for recording and rewind/resim
-			NetworkPhysicsComponent->CreateDataHistory<UE::ChaosMover::FNetworkDataTraits>(this);
-
-			if (NetworkPhysicsComponent->HasServerWorld())
+			else
 			{
-				if (APawn* PawnOwner = Cast<APawn>(GetOwner()))
+				NetworkPhysicsComponent = NewObject<UNetworkPhysicsComponent>(GetOwner(), TEXT("PhysMover_NetworkPhysicsComponent"));
+
+				// This isn't technically a DSO component, but set it net addressable as though it is
+				NetworkPhysicsComponent->SetNetAddressable();
+				NetworkPhysicsComponent->SetIsReplicated(true);
+				NetworkPhysicsComponent->RegisterComponent();
+				if (!NetworkPhysicsComponent->HasBeenInitialized())
 				{
-					// When we're owned by a pawn, keep an eye on whether it's currently player-controlled or not
-					PawnOwner->ReceiveControllerChangedDelegate.AddUniqueDynamic(this, &ThisClass::HandleOwningPawnControllerChanged_Server);
-					HandleOwningPawnControllerChanged_Server(PawnOwner, nullptr, PawnOwner->Controller);
+					NetworkPhysicsComponent->InitializeComponent();
 				}
-				else
+				NetworkPhysicsComponent->Activate(true);
+
+				// Register network data for recording and rewind/resim
+				NetworkPhysicsComponent->CreateDataHistory<UE::ChaosMover::FNetworkDataTraits>(this);
+
+				if (NetworkPhysicsComponent->HasServerWorld())
 				{
-					// If the owner isn't a pawn, there's no chance of player input happening, so inputs to the PT are always produced on the server
-					NetworkPhysicsComponent->SetIsRelayingLocalInputs(true);
+					if (APawn* PawnOwner = Cast<APawn>(GetOwner()))
+					{
+						// When we're owned by a pawn, keep an eye on whether it's currently player-controlled or not
+						PawnOwner->ReceiveControllerChangedDelegate.AddUniqueDynamic(this, &ThisClass::HandleOwningPawnControllerChanged_Server);
+						HandleOwningPawnControllerChanged_Server(PawnOwner, nullptr, PawnOwner->Controller);
+					}
+					else
+					{
+						// If the owner isn't a pawn, there's no chance of player input happening, so inputs to the PT are always produced on the server
+						NetworkPhysicsComponent->SetIsRelayingLocalInputs(true);
+					}
 				}
 			}
 		}
@@ -104,26 +113,48 @@ void UChaosMoverBackendComponent::UninitializeComponent()
 }
 
 void UChaosMoverBackendComponent::CreatePhysics()
+{
+	// Prevent the character particle from sleeping
+	Chaos::FPhysicsSolver* Solver = GetPhysicsSolver();
+	Chaos::FPBDRigidParticle* ControlledParticle = GetControlledParticle();
+	Chaos::FSingleParticlePhysicsProxy* ControlledParticleProxy = ControlledParticle ? static_cast<Chaos::FSingleParticlePhysicsProxy*>(ControlledParticle->GetProxy()) : nullptr;
+	if (ControlledParticleProxy)
+	{
+		ControlledParticle->SetSleepType(Chaos::ESleepType::NeverSleep);
+	}
+
+	// Create all possible constraints...
+	// ... a character ground constraint, for constraint based character-like movement on ground
+	CreateCharacterGroundConstraint();
+	// ... a path target constraint, for constraint based pathed movement
+	CreatePathTargetConstraint();
+}
+
+void UChaosMoverBackendComponent::DestroyPhysics()
+{
+	// Destroy all constraints
+	DestroyCharacterGroundConstraint();
+	DestroyPathTargetConstraint();
+}
+
+void UChaosMoverBackendComponent::CreateCharacterGroundConstraint()
 {
 	if (Chaos::FPhysicsSolver* Solver = GetPhysicsSolver())
 	{
-		if (Chaos::FPBDRigidParticle* Particle = GetControlledParticle())
+		if (Chaos::FPBDRigidParticle* ControlledParticle = GetControlledParticle())
 		{
-			if (Chaos::FSingleParticlePhysicsProxy* Proxy = static_cast<Chaos::FSingleParticlePhysicsProxy*>(Particle->GetProxy()))
+			if (Chaos::FSingleParticlePhysicsProxy* ControlledParticleProxy = static_cast<Chaos::FSingleParticlePhysicsProxy*>(ControlledParticle->GetProxy()))
 			{
+				// Create the character ground constraint, for character-like movement on ground
 				CharacterGroundConstraint = MakeUnique<Chaos::FCharacterGroundConstraint>();
-				CharacterGroundConstraint->Init(Proxy);
-
+				CharacterGroundConstraint->Init(ControlledParticleProxy);
 				Solver->RegisterObject(CharacterGroundConstraint.Get());
-
-				// Prevent the character particle from sleeping
-				Particle->SetSleepType(Chaos::ESleepType::NeverSleep);
 			}
 		}
 	}
 }
 
-void UChaosMoverBackendComponent::DestroyPhysics()
+void UChaosMoverBackendComponent::DestroyCharacterGroundConstraint()
 {
 	if (Chaos::FPhysicsSolver* Solver = GetPhysicsSolver())
 	{
@@ -135,6 +166,49 @@ void UChaosMoverBackendComponent::DestroyPhysics()
 	}
 }
 
+void UChaosMoverBackendComponent::CreatePathTargetConstraint()
+{
+	const UMoverComponent& MoverComp = GetMoverComponent();
+	if (Chaos::FPhysicsObject* PhysicsObject = GetPhysicsObject())
+	{
+		const FTransform ComponentWorldTransform = MoverComp.GetUpdatedComponent()->GetComponentTransform();
+		// Create the constraint via FChaosEngineInterface directly because it allows jointing a "real" object with a point in space (it creates a dummy particle for us)
+		FPhysicsConstraintHandle Handle = FChaosEngineInterface::CreateConstraint(PhysicsObject, nullptr, FTransform::Identity, FTransform::Identity);
+
+		bool bIsConstraintValid = false;
+		if (Handle.IsValid() && ensure(Handle->IsType(Chaos::EConstraintType::JointConstraintType)))
+		{
+			if (Chaos::FJointConstraint* Constraint = static_cast<Chaos::FJointConstraint*>(Handle.Constraint))
+			{
+				// Since we didn't use the ConstraintInstance to actually create the constraint (it requires both bodies exist, see comment above), link everything up manually
+				PathTargetConstraintHandle = Handle;						
+				PathTargetConstraintInstance.ConstraintHandle = PathTargetConstraintHandle;
+				Constraint->SetUserData(&PathTargetConstraintPhysicsUserData/*has a (void*)FConstraintInstanceBase*/);
+				bIsConstraintValid = true;
+
+				if (Chaos::FPBDRigidParticle* EndpointParticle = Constraint->GetPhysicsBodies()[1]->GetParticle<Chaos::EThreadContext::External>()->CastToRigidParticle())
+				{
+					EndpointParticle->SetX(ComponentWorldTransform.GetLocation());
+					EndpointParticle->SetR(ComponentWorldTransform.GetRotation());
+				}
+			}
+		}
+
+		if (!bIsConstraintValid)
+		{
+			FChaosEngineInterface::ReleaseConstraint(Handle);
+		}
+	}
+}
+
+void UChaosMoverBackendComponent::DestroyPathTargetConstraint()
+{
+	if (PathTargetConstraintHandle.IsValid())
+	{
+		FChaosEngineInterface::ReleaseConstraint(PathTargetConstraintHandle);
+	}
+}
+
 void UChaosMoverBackendComponent::HandleOwningPawnControllerChanged_Server(APawn* OwnerPawn, AController* OldController, AController* NewController)
 {
 	// Inputs for player-controlled pawns originate on the player's client, all others originate on the server
@@ -173,14 +247,17 @@ UMoverComponent& UChaosMoverBackendComponent::GetMoverComponent() const
 	return *GetOuterUMoverComponent();
 }
 
+Chaos::FPhysicsObject* UChaosMoverBackendComponent::GetPhysicsObject() const
+{
+	IPhysicsComponent* PhysicsComponent = Cast<IPhysicsComponent>(GetMoverComponent().GetUpdatedComponent());
+	return PhysicsComponent ? PhysicsComponent->GetPhysicsObjectByName(NAME_None) : nullptr;
+}
+
 Chaos::FPBDRigidParticle* UChaosMoverBackendComponent::GetControlledParticle() const
 {
-	if (IPhysicsComponent* PhysicsComponent = Cast<IPhysicsComponent>(GetMoverComponent().GetUpdatedComponent()))
+	if (Chaos::FPhysicsObject* PhysicsObject = GetPhysicsObject())
 	{
-		if (Chaos::FPhysicsObject* PhysObj = PhysicsComponent->GetPhysicsObjectByName(NAME_None))
-		{
-			return FPhysicsObjectExternalInterface::LockRead(PhysObj)->GetRigidParticle(PhysObj);
-		}
+		return FPhysicsObjectExternalInterface::LockRead(PhysicsObject)->GetRigidParticle(PhysicsObject);
 	}
 
 	return nullptr;
@@ -204,13 +281,9 @@ void UChaosMoverBackendComponent::InitSimulation()
 	Params.NullMovementMode = TWeakObjectPtr<UNullMovementMode>(NullMovementMode.Get());
 	Params.StartingMovementMode = MoverComp.StartingMovementMode;
 	Params.CharacterConstraintProxy = CharacterGroundConstraint ? CharacterGroundConstraint->GetProxy<Chaos::FCharacterGroundConstraintProxy>() : nullptr;
-	if (IPhysicsComponent* PhysicsComponent = Cast<IPhysicsComponent>(MoverComp.GetUpdatedComponent()))
-	{
-		if (Chaos::FPhysicsObjectHandle PhysObj = PhysicsComponent->GetPhysicsObjectByName(NAME_None))
-		{
-			Params.PhysicsObject = PhysObj;
-		}
-	}
+	Params.PathTargetConstraintProxy = PathTargetConstraintHandle.IsValid() ? PathTargetConstraintHandle->GetProxy<Chaos::FJointConstraintPhysicsProxy>() : nullptr;
+	Params.PathTargetKinematicEndPointProxy = Params.PathTargetConstraintProxy && Params.PathTargetConstraintProxy->GetConstraint() ? Params.PathTargetConstraintProxy->GetConstraint()->GetKinematicEndPoint() : nullptr;
+	Params.PhysicsObject = GetPhysicsObject();
 	Params.Solver = GetPhysicsSolver();
 	Params.World = GetWorld();
 	
@@ -398,16 +471,6 @@ void UChaosMoverBackendComponent::FinalizeFrame(double ResultsTimeInSeconds)
 		TransformSyncState.MoveDirectionIntent = TransformSyncState.GetCapturedMovementBaseQuat().UnrotateVector(TransformSyncState.MoveDirectionIntent);
 	}
 
-	// Send the latest physics thread data to CVD
-	// Note we are sending the latest physics results, not the interpolated results
-#if WITH_CHAOS_VISUAL_DEBUGGER
-	if (FChaosVisualDebuggerTrace::IsTracing())
-	{
-		const UE::ChaosMover::FSimulationOutputData& LatestOutput = SimOutputRecord.GetLatest();
-		FMoverCVDRuntimeTrace::TraceMoverData(&MoverComp, &InterpolatedOutput.LastUsedInputCmd, &InterpolatedOutput.SyncState, &Simulation->GetLocalSimInput());
-	}
-#endif
-
 	// This eventually will process any events queued by the simulation during the async tick.
 	// For now it's just movement mode activation, deactivation and transition callbacks.
 	FlushModeChangeEvents();
@@ -425,7 +488,8 @@ void UChaosMoverBackendComponent::FinalizeFrame(double ResultsTimeInSeconds)
 */
 void UChaosMoverBackendComponent::FlushModeChangeEvents()
 {
-	const TArray<UE::ChaosMover::FMovementModeChangeEvent>& QueuedModeChangeEvents = Simulation->GetQueuedModeChangeEvents();
+	TArray<UE::ChaosMover::FMovementModeChangeEvent> QueuedModeChangeEvents;
+	Simulation->GetAndClearQueuedModeChangeEvents(QueuedModeChangeEvents);
 	UMoverComponent& MoverComp = GetMoverComponent();
 
 	for (int Idx = QueuedModeChangeEvents.Num() - 1; Idx >= 0; --Idx)
@@ -447,7 +511,5 @@ void UChaosMoverBackendComponent::FlushModeChangeEvents()
 		// Signal movement mode change event
 		MoverComp.OnMovementModeChanged.Broadcast(MovementModeChangeEvent.PreviousMovementModeName, MovementModeChangeEvent.NextMovementModeName);
 	}
-
-	Simulation->ClearQueuedModeChangeEvents();
 }
 
diff --git a/Engine/Plugins/Experimental/ChaosMover/Source/ChaosMover/Private/ChaosMoverSimulation.cpp b/Engine/Plugins/Experimental/ChaosMover/Source/ChaosMover/Private/ChaosMoverSimulation.cpp
--- a/Engine/Plugins/Experimental/ChaosMover/Source/ChaosMover/Private/ChaosMoverSimulation.cpp
+++ b/Engine/Plugins/Experimental/ChaosMover/Source/ChaosMover/Private/ChaosMoverSimulation.cpp
@@ -6,6 +6,8 @@
 #include "Chaos/Character/CharacterGroundConstraintContainer.h"
 #include "Chaos/ContactModification.h"
 #include "Chaos/PhysicsObjectInternalInterface.h"
+#include "PhysicsEngine/PhysicsObjectExternalInterface.h"
+#include "Physics/Experimental/PhysScene_Chaos.h"
 #include "ChaosMover/ChaosMovementMode.h"
 #include "ChaosMover/ChaosMovementModeTransition.h"
 #include "ChaosMover/ChaosMoverLog.h"
@@ -20,7 +22,16 @@
 #include "MoverSimulationTypes.h"
 #include "PBDRigidsSolver.h"
 #include "PhysicsProxy/CharacterGroundConstraintProxy.h"
+#include "Chaos/PBDJointConstraintData.h"
 #include "UObject/UObjectGlobals.h"
+#include "Chaos/ChaosEngineInterface.h"
+#include "Chaos/KinematicTargets.h"
+#include "Engine/World.h"
+
+#if WITH_CHAOS_VISUAL_DEBUGGER
+#include "ChaosVisualDebugger/ChaosVisualDebuggerTrace.h"
+#include "ChaosVisualDebugger/MoverCVDRuntimeTrace.h"
+#endif
 
 #include UE_INLINE_GENERATED_CPP_BY_NAME(ChaosMoverSimulation)
 
@@ -41,6 +52,11 @@ FMoverDataCollection& UChaosMoverSimulation::GetLocalSimInput_Mutable()
 	return LocalSimInput;
 }
 
+FMoverDataCollection& UChaosMoverSimulation::GetDebugSimData()
+{
+	return DebugSimData;
+}
+
 const UBaseMovementMode* UChaosMoverSimulation::GetCurrentMovementMode() const
 {
 	return StateMachine.GetCurrentMode().Get();
@@ -79,8 +95,13 @@ void UChaosMoverSimulation::BuildNetStateData(FMoverSyncState& OutNetSyncState)
 
 void UChaosMoverSimulation::Init(const FInitParams& InitParams)
 {
+	// Only the Gameplay Thread is allowed to Init the chaos mover simulation
+	Chaos::EnsureIsInGameThreadContext();
+
 	MovementMixerWeakPtr = InitParams.MovementMixer.Get();
 	CharacterConstraintProxy = InitParams.CharacterConstraintProxy;
+	PathTargetConstraintProxy = InitParams.PathTargetConstraintProxy;
+	PathTargetKinematicEndPointProxy = InitParams.PathTargetKinematicEndPointProxy;
 	PhysicsObject = InitParams.PhysicsObject;
 	Solver = InitParams.Solver;
 	World = InitParams.World;
@@ -95,29 +116,27 @@ void UChaosMoverSimulation::Init(const FInitParams& InitParams)
 
 	for (const TPair<FName, TWeakObjectPtr<UBaseMovementMode>>& Element : InitParams.ModesToRegister)
 	{
-		if (Element.Value.Get() == nullptr)
+		const FName& ModeName = Element.Key;
+		TWeakObjectPtr<UBaseMovementMode> Mode = Element.Value;
+
+		if (Mode.Get() == nullptr)
 		{
-			UE_LOG(LogChaosMover, Warning, TEXT("Invalid Movement Mode type '%s' detected. Mover actor will not function correctly."), *Element.Key.ToString());
+			UE_LOG(LogChaosMover, Warning, TEXT("Invalid Movement Mode type '%s' detected. Mover actor will not function correctly."), *ModeName.ToString());
 			continue;
 		}
 
-		if (UChaosMovementMode* ChaosMode = Cast<UChaosMovementMode>(Element.Value.Get()))
+		if (UChaosMovementMode* ChaosMode = Cast<UChaosMovementMode>(Mode.Get()))
 		{
 			ChaosMode->SetSimulation(this);
 		}
 
-		StateMachine.RegisterMovementMode(Element.Key, Element.Value);
-	}
-
-	if (InitParams.StartingMovementMode != NAME_None)
-	{
-		StateMachine.SetDefaultMode(InitParams.StartingMovementMode);
-		StateMachine.SetModeImmediately(InitParams.StartingMovementMode);
+		bool bIsDefaultMode = (InitParams.StartingMovementMode == ModeName);
+		StateMachine.RegisterMovementMode(ModeName, Mode, bIsDefaultMode);
 	}
 
 	for (const TWeakObjectPtr<UBaseMovementModeTransition>& Transition : InitParams.TransitionsToRegister)
 	{
-		if (UChaosMovementModeTransition* ChaosTransition= Cast<UChaosMovementModeTransition>(Transition.Get()))
+		if (UChaosMovementModeTransition* ChaosTransition = Cast<UChaosMovementModeTransition>(Transition.Get()))
 		{
 			ChaosTransition->SetSimulation(this);
 		}
@@ -165,7 +184,7 @@ void UChaosMoverSimulation::ModifyContacts(const FMoverTimeStep& TimeStep, const
 			if (ChaosMode->IgnoredCollisionMode == EChaosMoverIgnoredCollisionMode::DisableCollisionsWithIgnored)
 			{
 				Chaos::FGeometryParticleHandle* UpdatedParticle = nullptr;
-				const FChaosMoverSimulationDefaultInputs* SimInputs = InputData.LocalInputParams.AuxStateCollection.FindDataByType<FChaosMoverSimulationDefaultInputs>();
+				const FChaosMoverSimulationDefaultInputs* SimInputs = LocalSimInput.FindDataByType<FChaosMoverSimulationDefaultInputs>();
 				if (SimInputs)
 				{
 					Chaos::FReadPhysicsObjectInterface_Internal ReadInterface = Chaos::FPhysicsObjectInternalInterface::GetRead();
@@ -249,10 +268,10 @@ void UChaosMoverSimulation::OnPreSimulationTick(const FMoverTimeStep& TimeStep,
 	}
 
 	// Update the sync state from the current physics state
-	FMoverDefaultSyncState& DefaultSyncState = CurrentSyncState.SyncStateCollection.FindOrAddMutableDataByType<FMoverDefaultSyncState>();
-	if (const Chaos::FPBDRigidParticleHandle* ParticleHandle = UChaosMoverSimulation::GetControlledParticle())
+	FMoverDefaultSyncState& PreSimDefaultSyncState = CurrentSyncState.SyncStateCollection.FindOrAddMutableDataByType<FMoverDefaultSyncState>();
+	if (const Chaos::FPBDRigidParticleHandle* ParticleHandle = GetControlledParticle())
 	{
-		DefaultSyncState.SetTransforms_WorldSpace(ParticleHandle->GetX(), FRotator(ParticleHandle->GetR()), ParticleHandle->GetV());
+		PreSimDefaultSyncState.SetTransforms_WorldSpace(ParticleHandle->GetX(), FRotator(ParticleHandle->GetR()), ParticleHandle->GetV());
 	}
 
 	if (TStrongObjectPtr<UBaseMovementMode> CurrentModePtr = StateMachine.GetCurrentMode().Pin())
@@ -281,21 +300,21 @@ void UChaosMoverSimulation::OnSimulationTick(const FMoverTimeStep& TimeStep, con
 void UChaosMoverSimulation::PostSimulationTickCharacter(const IChaosCharacterMovementModeInterface& CharacterMode, const FMoverTimeStep& TimeStep, UE::ChaosMover::FSimulationOutputData& OutputData)
 {
 	FChaosMoverCharacterSimState& CharacterSimState = InternalSimData.FindOrAddMutableDataByType<FChaosMoverCharacterSimState>();
-	const FMoverDefaultSyncState* DefaultSyncState = OutputData.SyncState.SyncStateCollection.FindDataByType<FMoverDefaultSyncState>();
-	check(DefaultSyncState);
+	FMoverDefaultSyncState* PostSimDefaultSyncState = OutputData.SyncState.SyncStateCollection.FindMutableDataByType<FMoverDefaultSyncState>();
+	check(PostSimDefaultSyncState);
 
-	if (Chaos::FPBDRigidParticleHandle* ParticleHandle = UChaosMoverSimulation::GetControlledParticle())
+	if (Chaos::FPBDRigidParticleHandle* ParticleHandle = GetControlledParticle())
 	{
 		// Store the delta position and rotation locally
-		FRotator DeltaRotation = DefaultSyncState->GetOrientation_WorldSpace() - FRotator(ParticleHandle->GetR());
+		FRotator DeltaRotation = PostSimDefaultSyncState->GetOrientation_WorldSpace() - FRotator(ParticleHandle->GetR());
 		FRotator Winding, Remainder;
 		DeltaRotation.GetWindingAndRemainder(Winding, Remainder);
 
-		CharacterSimState.TargetDeltaPosition = DefaultSyncState->GetLocation_WorldSpace() - ParticleHandle->GetX();
+		CharacterSimState.TargetDeltaPosition = PostSimDefaultSyncState->GetLocation_WorldSpace() - ParticleHandle->GetX();
 		CharacterSimState.TargetDeltaFacing = FMath::DegreesToRadians(Remainder.Yaw);
 
 		// Set the particle velocity to hit the target sync state transform
-		ParticleHandle->SetV(DefaultSyncState->GetVelocity_WorldSpace());
+		ParticleHandle->SetV(PostSimDefaultSyncState->GetVelocity_WorldSpace());
 		const float DeltaTimeSeconds = TimeStep.StepMs * 0.001f;
 		if (DeltaTimeSeconds > UE_SMALL_NUMBER)
 		{
@@ -304,7 +323,7 @@ void UChaosMoverSimulation::PostSimulationTickCharacter(const IChaosCharacterMov
 			if (!CharacterMode.ShouldCharacterRemainUpright())
 			{
 				FVector PreviousUp = ParticleHandle->GetR().RotateVector(FVector::UpVector).GetSafeNormal();
-				FVector TargetUp = DefaultSyncState->GetOrientation_WorldSpace().Quaternion().RotateVector(FVector::UpVector).GetSafeNormal();
+				FVector TargetUp = PostSimDefaultSyncState->GetOrientation_WorldSpace().Quaternion().RotateVector(FVector::UpVector).GetSafeNormal();
 				AngularVelocity += (PreviousUp.Cross(TargetUp).GetSafeNormal() * FMath::Acos(TargetUp.Dot(PreviousUp))) / DeltaTimeSeconds;
 			}
 
@@ -316,6 +335,15 @@ void UChaosMoverSimulation::PostSimulationTickCharacter(const IChaosCharacterMov
 		CharacterSimState.TargetDeltaPosition = FVector::ZeroVector;
 		CharacterSimState.TargetDeltaFacing = 0.0f;
 	}
+
+	// Update the movement base
+	check(Blackboard.Get());
+	FFloorCheckResult FloorResult;
+	bool FoundLastFloorResult = Blackboard->TryGet(CommonBlackboard::LastFloorResult, FloorResult);
+	bool FoundFloor = FoundLastFloorResult && FloorResult.bBlockingHit;
+	// Note: We want to record the movement base but we don't record the transform
+	// so don't use this to get a relative transform for the sync state
+	PostSimDefaultSyncState->SetMovementBase(FoundFloor ? FloorResult.HitResult.GetComponent() : nullptr);
 }
 
 void UChaosMoverSimulation::PostSimulationTickCharacterConstraint(const IChaosCharacterConstraintMovementModeInterface& CharacterConstraintMode, const FMoverTimeStep& TimeStep, UE::ChaosMover::FSimulationOutputData& OutputData)
@@ -340,16 +368,6 @@ void UChaosMoverSimulation::PostSimulationTickCharacterConstraint(const IChaosCh
 		return;
 	}
 
-	// Update the constraint settings if the mode changes
-	// TODO - Add an initialize for this so we don't have to set it every frame
-	//if (CurrentSyncState.MovementMode != OutputData.SyncState.MovementMode)
-	{
-		Chaos::FCharacterGroundConstraintSettings& Settings = ConstraintHandle->GetSettings_Mutable();
-		CharacterConstraintMode.UpdateConstraintSettings(Settings);
-	}
-
-	FMoverDefaultSyncState* DefaultSyncState = OutputData.SyncState.SyncStateCollection.FindMutableDataByType<FMoverDefaultSyncState>();
-	check(DefaultSyncState);
 	check(Blackboard.Get());
 
 	FFloorCheckResult FloorResult;
@@ -398,15 +416,10 @@ void UChaosMoverSimulation::PostSimulationTickCharacterConstraint(const IChaosCh
 			FloorResult.FloorDist,
 			WalkableSlopeCosine
 			});
-
-		// Note: We want to record the movement base but we don't record the transform
-		// so don't use this to get a relative transform for the sync state
-		DefaultSyncState->SetMovementBase(FloorResult.HitResult.GetComponent());
 	}
 	else
 	{
 		ConstraintHandle->SetGroundParticle(nullptr);
-		DefaultSyncState->SetMovementBase(nullptr);
 		ConstraintHandle->SetData({
 			ConstraintHandle->GetSettings().VerticalAxis,
 			Chaos::FVec3::ZeroVector,
@@ -417,6 +430,64 @@ void UChaosMoverSimulation::PostSimulationTickCharacterConstraint(const IChaosCh
 	}
 }
 
+void UChaosMoverSimulation::PostSimulationTickPathedMovement(const IChaosPathedMovementModeInterface& ConstraintMode, const FMoverTimeStep& TimeStep, UE::ChaosMover::FSimulationOutputData& OutputData)
+{
+	Chaos::FPBDJointConstraintHandle* ConstraintHandle = nullptr;
+	if (PathTargetConstraintProxy && PathTargetConstraintProxy->IsInitialized())
+	{
+		ConstraintHandle = PathTargetConstraintProxy->GetHandle();
+	}
+	if (!ConstraintHandle)
+	{
+		return;
+	}
+
+	if (ConstraintMode.ShouldUseConstraint())
+	{
+		if (!ConstraintHandle->IsEnabled())
+		{
+			EnablePathTargetConstraint();
+		}
+
+		if (!IsControlledParticleDynamic())
+		{
+			SetControlledParticleDynamic();
+		}
+	}
+	else if (!ConstraintMode.ShouldUseConstraint())
+	{
+		if (ConstraintHandle->IsEnabled())
+		{
+			DisablePathTargetConstraint();
+		}
+
+		if (!IsControlledParticleKinematic())
+		{
+			SetControlledParticleKinematic();
+		}
+	}
+
+	const FMoverDefaultSyncState* PostSimDefaultSyncState = OutputData.SyncState.SyncStateCollection.FindMutableDataByType<FMoverDefaultSyncState>();
+	check(PostSimDefaultSyncState);
+	check(Solver);
+	Chaos::FPBDRigidsEvolution& Evolution = *Solver->GetEvolution();
+
+	// We always update the path target end point, even when the controlled particle is not constrained.
+	// That way we don't leave a particle behind and don't risk having it very far and whip lashing when re-enabling the constraint.
+	Chaos::FKinematicGeometryParticleHandle* EndpointParticleHandle = (ConstraintHandle && PathTargetKinematicEndPointProxy) ? PathTargetKinematicEndPointProxy->GetHandle_LowLevel()->CastToKinematicParticle() : nullptr;
+	if (EndpointParticleHandle)
+	{	
+		Evolution.SetParticleKinematicTarget(EndpointParticleHandle, Chaos::FKinematicTarget::MakePositionTarget(PostSimDefaultSyncState->GetTransform_WorldSpace()));
+	}
+	if (IsControlledParticleKinematic())
+	{
+		if (Chaos::FPBDRigidParticleHandle* ParticleHandle = GetControlledParticle())
+		{
+			Evolution.SetParticleKinematicTarget(ParticleHandle, Chaos::FKinematicTarget::MakePositionTarget(PostSimDefaultSyncState->GetTransform_WorldSpace()));
+		}
+	}
+}
+
 void UChaosMoverSimulation::OnPostSimulationTick(const FMoverTimeStep& TimeStep, UE::ChaosMover::FSimulationOutputData& OutputData)
 {
 	// TODO - make this more extensible
@@ -435,23 +506,109 @@ void UChaosMoverSimulation::OnPostSimulationTick(const FMoverTimeStep& TimeStep,
 		{
 			DisableCharacterConstraint();
 		}
+
+		if (IChaosPathedMovementModeInterface* ConstraintMode = Cast<IChaosPathedMovementModeInterface>(CurrentModePtr.Get()))
+		{
+			PostSimulationTickPathedMovement(*ConstraintMode, TimeStep, OutputData);
+		}
+		else
+		{
+			DisablePathTargetConstraint();
+		}
 	}
 
 	CurrentSyncState = OutputData.SyncState;
+
+	// Send Debug Data to the Chaos Visual Debugger
+	TraceMoverData(OutputData);
+}
+
+void UChaosMoverSimulation::TraceMoverData(const UE::ChaosMover::FSimulationOutputData& OutputData)
+{
+	// Send the latest physics thread data to CVD
+#if WITH_CHAOS_VISUAL_DEBUGGER
+	if (FChaosVisualDebuggerTrace::IsTracing())
+	{
+		static FName NAME_LocalSimImput("LocalSimImput");
+		static FName NAME_InternalSimData("InternalSimData");
+		static FName NAME_DebugSimData("DebugSimData");
+
+		UE::MoverUtils::NamedDataCollections LocalSimDataCollections (
+			{				
+				{ NAME_LocalSimImput, &LocalSimInput},
+				{ NAME_InternalSimData, &InternalSimData},
+				{ NAME_DebugSimData, &DebugSimData}
+			});
+
+		Chaos::FReadPhysicsObjectInterface_Internal Interface = Chaos::FPhysicsObjectInternalInterface::GetRead();
+		const Chaos::FGeometryParticleHandle* ParticleHandle = PhysicsObject? Interface.GetParticle(PhysicsObject) : nullptr;
+		int32 ParticleID = ParticleHandle ? ParticleHandle->UniqueIdx().Idx : INDEX_NONE;
+
+		int32 SolverID = CVD_TRACE_GET_SOLVER_ID_FROM_WORLD(World);
+
+		UE::MoverUtils::FMoverCVDRuntimeTrace::TraceMoverData(SolverID, ParticleID, &OutputData.LastUsedInputCmd, &OutputData.SyncState, &LocalSimDataCollections);
+	}
+#endif
 }
 
 void UChaosMoverSimulation::OnModifyContacts(const FMoverTimeStep& TimeStep, const UE::ChaosMover::FSimulationInputData& InputData, const UE::ChaosMover::FSimulationOutputData& OutputData, Chaos::FCollisionContactModifier& Modifier) const
 {
 }
 
-const TArray<UE::ChaosMover::FMovementModeChangeEvent>& UChaosMoverSimulation::GetQueuedModeChangeEvents() const
+void UChaosMoverSimulation::OnMovementModeChanged(FName PreviousMovementModeName, TWeakObjectPtr<UBaseMovementMode> PreviousMovementMode, FName NextMovementModeName, TWeakObjectPtr<UBaseMovementMode> NextMovementMode)
+{
+	// Queue a mode changed event so it can be processed on the game thread to notify gameplay
+	EnqueueModeChangeEvent(/*FMovementModeChangeEvent*/{ PreviousMovementModeName, PreviousMovementMode, NextMovementModeName, NextMovementMode });
+
+	TStrongObjectPtr<UBaseMovementMode> NextModePtr = NextMovementMode.Pin();
+	TStrongObjectPtr<UBaseMovementMode> PreviousModePtr = PreviousMovementMode.Pin();
+
+	if (IChaosCharacterConstraintMovementModeInterface* NextCharacterConstraintMode = Cast<IChaosCharacterConstraintMovementModeInterface>(NextModePtr.Get()))
+	{
+		bool IsCharacterConstraintInitialized = (CharacterConstraintProxy != nullptr) && CharacterConstraintProxy->IsInitialized();
+		Chaos::FCharacterGroundConstraintHandle* ConstraintHandle = IsCharacterConstraintInitialized ? CharacterConstraintProxy->GetPhysicsThreadAPI() : nullptr;
+		if (ConstraintHandle)
+		{
+			Chaos::FCharacterGroundConstraintSettings& Settings = ConstraintHandle->GetSettings_Mutable();
+			NextCharacterConstraintMode->UpdateConstraintSettings(Settings);
+
+			// Character ground constraint modes currently assume moving a dynamic particle and using a character ground constraint
+			// Revise if we start supporting moving a character kinematically
+			if (!IsControlledParticleDynamic())
+			{
+				SetControlledParticleDynamic();
+			}
+		}
+	}
+
+	if (IChaosPathedMovementModeInterface* NextPathTargetConstraintMode = Cast<IChaosPathedMovementModeInterface>(NextModePtr.Get()))
+	{
+		bool IsPathTargetConstraintInitialized = (PathTargetConstraintProxy != nullptr) && PathTargetConstraintProxy->IsInitialized();
+		Chaos::FPBDJointConstraintHandle* ConstraintHandle = IsPathTargetConstraintInitialized ? PathTargetConstraintProxy->GetHandle() : nullptr;
+		if (ConstraintHandle)
+		{
+			ConstraintHandle->SetSettings(NextPathTargetConstraintMode->GetConstraintSettings());
+		}
+	}
+}
+
+void UChaosMoverSimulation::EnqueueModeChangeEvent(const UE::ChaosMover::FMovementModeChangeEvent& MovementModeChangeEvent)
 {
-	return StateMachine.GetQueuedModeChangeEvents();
+	// FinalizeFrame and the next physics tick can be concurrent so this was not thread safe
+	// This mutex lock is just temporary to prevent crashes
+	// TODO: Implement a lockless solution (output the mode changed events as part of sim output instead)
+	FRWScopeLock Lock(ModeChangedEventsLock, SLT_Write);
+	QueuedModeChangeEvents.Add(MovementModeChangeEvent);
 }
 
-void UChaosMoverSimulation::ClearQueuedModeChangeEvents()
+void UChaosMoverSimulation::GetAndClearQueuedModeChangeEvents(TArray<UE::ChaosMover::FMovementModeChangeEvent>& OutModeChangedEvents)
 {
-	return StateMachine.ClearQueuedModeChangeEvents();
+	// FinalizeFrame and the next physics tick can be concurrent so this was not thread safe
+	// This mutex lock is just temporary to prevent crashes
+	// TODO: Implement a lockless solution (output the mode changed events as part of sim output instead)
+	FRWScopeLock Lock(ModeChangedEventsLock, SLT_ReadOnly);
+	OutModeChangedEvents = QueuedModeChangeEvents;
+	QueuedModeChangeEvents.Empty();
 }
 
 Chaos::FPBDRigidParticleHandle* UChaosMoverSimulation::GetControlledParticle() const
@@ -464,12 +621,52 @@ Chaos::FPBDRigidParticleHandle* UChaosMoverSimulation::GetControlledParticle() c
 	return nullptr;
 }
 
+void UChaosMoverSimulation::SetControlledParticleDynamic()
+{
+	if (Chaos::FPBDRigidParticleHandle* ControlledParticle = GetControlledParticle())
+	{
+		check(Solver);
+		Chaos::FPBDRigidsEvolution& Evolution = *Solver->GetEvolution();
+		Evolution.SetParticleObjectState(ControlledParticle, Chaos::EObjectStateType::Dynamic);
+	}
+}
+
+void UChaosMoverSimulation::SetControlledParticleKinematic()
+{
+	if (Chaos::FPBDRigidParticleHandle* ControlledParticle = GetControlledParticle())
+	{
+		check(Solver);
+		Chaos::FPBDRigidsEvolution& Evolution = *Solver->GetEvolution();
+		Evolution.SetParticleObjectState(ControlledParticle, Chaos::EObjectStateType::Kinematic);
+
+		if (const Chaos::TPBDRigidParticleHandleImp<Chaos::FReal, 3, true>* ControlledRigidParticle = ControlledParticle->CastToRigidParticle())
+		{
+			if (ControlledRigidParticle->UpdateKinematicFromSimulation() == true)
+			{
+				// Should we instead call SetUpdateKinematicFromSimulation on the GT when some of the modes may animate kinematically?
+				UE_LOG(LogChaosMover, Warning, TEXT("The updated component for %s is not set to Update Kinematic from Simulation but is being moved kinematically"), *GetClass()->GetName());
+			}
+		}
+	}
+}
+
+bool UChaosMoverSimulation::IsControlledParticleDynamic() const
+{
+	const Chaos::FPBDRigidParticleHandle* ControlledParticle = GetControlledParticle();
+	return ControlledParticle ? ControlledParticle->IsDynamic() : false;
+}
+
+bool UChaosMoverSimulation::IsControlledParticleKinematic() const
+{
+	const Chaos::FPBDRigidParticleHandle* ControlledParticle = GetControlledParticle();
+	return ControlledParticle ? ControlledParticle->IsKinematic() : false;
+}
+
 void UChaosMoverSimulation::EnableCharacterConstraint()
 {
-	Chaos::FCharacterGroundConstraintHandle* ConstraintHandle;
 	if (CharacterConstraintProxy && CharacterConstraintProxy->IsInitialized())
 	{
-		ConstraintHandle = CharacterConstraintProxy->GetPhysicsThreadAPI();
+		Chaos::FCharacterGroundConstraintHandle* ConstraintHandle = CharacterConstraintProxy->GetPhysicsThreadAPI();
 
 		if (ConstraintHandle->GetCharacterParticle())
 		{
@@ -480,15 +677,36 @@ void UChaosMoverSimulation::EnableCharacterConstraint()
 
 void UChaosMoverSimulation::DisableCharacterConstraint()
 {
-	Chaos::FCharacterGroundConstraintHandle* ConstraintHandle;
 	if (CharacterConstraintProxy && CharacterConstraintProxy->IsInitialized())
 	{
-		ConstraintHandle = CharacterConstraintProxy->GetPhysicsThreadAPI();
+		Chaos::FCharacterGroundConstraintHandle* ConstraintHandle = CharacterConstraintProxy->GetPhysicsThreadAPI();
 
 		ConstraintHandle->SetEnabled(false);
 	}
 }
 
+void UChaosMoverSimulation::EnablePathTargetConstraint()
+{
+	if (PathTargetConstraintProxy && PathTargetConstraintProxy->IsInitialized())
+	{
+		if (Chaos::FPBDJointConstraintHandle* ConstraintHandle = PathTargetConstraintProxy->GetHandle())
+		{
+			ConstraintHandle->SetConstraintEnabled(true);
+		}
+	}
+}
+
+void UChaosMoverSimulation::DisablePathTargetConstraint()
+{
+	if (PathTargetConstraintProxy && PathTargetConstraintProxy->IsInitialized())
+	{
+		if (Chaos::FPBDJointConstraintHandle* ConstraintHandle = PathTargetConstraintProxy->GetHandle())
+		{
+			ConstraintHandle->SetConstraintEnabled(false);
+		}
+	}
+}
+
 void UChaosMoverSimulation::K2_QueueInstantMovementEffect(const int32& EffectAsRawData)
 {
 	// This will never be called, the exec version below will be hit instead
@@ -520,7 +738,6 @@ DEFINE_FUNCTION(UChaosMoverSimulation::execK2_QueueInstantMovementEffect)
 	P_NATIVE_END;
 }
 
-
 void UChaosMoverSimulation::QueueInstantMovementEffect(TSharedPtr<FInstantMovementEffect> InstantMovementEffect)
 {
 	StateMachine.QueueInstantMovementEffect(InstantMovementEffect);
diff --git a/Engine/Plugins/Experimental/ChaosMover/Source/ChaosMover/Private/ChaosMoverSimulationTypes.cpp b/Engine/Plugins/Experimental/ChaosMover/Source/ChaosMover/Private/ChaosMoverSimulationTypes.cpp
--- a/Engine/Plugins/Experimental/ChaosMover/Source/ChaosMover/Private/ChaosMoverSimulationTypes.cpp
+++ b/Engine/Plugins/Experimental/ChaosMover/Source/ChaosMover/Private/ChaosMoverSimulationTypes.cpp
@@ -1,6 +1,7 @@
 // Copyright Epic Games, Inc. All Rights Reserved.
 
 #include "ChaosMover/ChaosMoverSimulationTypes.h"
+#include "ChaosMover/ChaosMoverSimulation.h"
 
 void FChaosMoverSimulationDefaultInputs::Reset()
 {
@@ -15,4 +16,10 @@ void FChaosMoverSimulationDefaultInputs::Reset()
 	OwningActor = nullptr;
 	World = nullptr;
 	CollisionChannel = ECC_Pawn;
+}
+
+FMoverDataCollection& UE::ChaosMover::GetDebugSimData(UChaosMoverSimulation* Simulation)
+{
+	check(Simulation);
+	return Simulation->GetDebugSimData();
 }
\ No newline at end of file
diff --git a/Engine/Plugins/Experimental/ChaosMover/Source/ChaosMover/Private/ChaosMoverStateMachine.cpp b/Engine/Plugins/Experimental/ChaosMover/Source/ChaosMover/Private/ChaosMoverStateMachine.cpp
--- a/Engine/Plugins/Experimental/ChaosMover/Source/ChaosMover/Private/ChaosMoverStateMachine.cpp
+++ b/Engine/Plugins/Experimental/ChaosMover/Source/ChaosMover/Private/ChaosMoverStateMachine.cpp
@@ -14,6 +14,7 @@
 #include "MoverDeveloperSettings.h"
 #include "MoverLog.h"
 #include "Templates/SubclassOf.h"
+#include "Framework/Threading.h"
 
 namespace UE::ChaosMover
 {
@@ -27,6 +28,9 @@ namespace UE::ChaosMover
 
 	void FMoverStateMachine::Init(const FInitParams& Params)
 	{
+		Chaos::EnsureIsInGameThreadContext();
+
+		// Careful, this is called from the GT
 		ImmediateMovementModeTransitionWeakPtr = Params.ImmediateMovementModeTransition;
 		NullMovementModeWeakPtr = Params.NullMovementMode;
 		Simulation = Params.Simulation;
@@ -86,6 +90,11 @@ namespace UE::ChaosMover
 		DefaultModeName = NewDefaultModeName;
 	}
 
+	FName FMoverStateMachine::GetDefaultModeName() const
+	{
+		return DefaultModeName;
+	}
+
 	void FMoverStateMachine::RegisterGlobalTransition(TWeakObjectPtr<UBaseMovementModeTransition> Transition)
 	{
 		GlobalTransitions.Add(Transition);
@@ -488,8 +497,12 @@ namespace UE::ChaosMover
 						UE_LOG(LogChaosMover, Verbose, TEXT("AdvanceToNextMode: %s (%s) from %s to %s"),
 							*OwnerActorName, *UEnum::GetValueAsString(OwnerActorLocalNetRole), *CurrentModeName.ToString(), *NextModeName.ToString());
 
+						// Notify the simulation of a mode change so it can react accordingly
 						TWeakObjectPtr<UBaseMovementMode> PreviousMovementMode = FoundPreviousMovementMode ? *FoundPreviousMovementMode : TWeakObjectPtr<UBaseMovementMode>();
-						QueuedModeChangeEvents.Add({ PreviousModeName, PreviousMovementMode, CurrentModeName, *FoundNextMovementMode });
+						if (TStrongObjectPtr<UChaosMoverSimulation> SimStrongObjPtr = Simulation.Pin())
+						{
+							SimStrongObjPtr->OnMovementModeChanged(PreviousModeName, PreviousMovementMode, CurrentModeName, *FoundNextMovementMode);
+						}
 					}
 				}
 			}
@@ -498,16 +511,6 @@ namespace UE::ChaosMover
 		}
 	}
 
-	const TArray<FMovementModeChangeEvent>& FMoverStateMachine::GetQueuedModeChangeEvents() const
-	{
-		return QueuedModeChangeEvents;
-	}
-
-	void FMoverStateMachine::ClearQueuedModeChangeEvents()
-	{
-		QueuedModeChangeEvents.Empty();
-	}
-
 	void FMoverStateMachine::FlushQueuedMovesToGroup(FLayeredMoveGroup& Group)
 	{
 		if (!QueuedLayeredMoves.IsEmpty())
diff --git a/Engine/Plugins/Experimental/ChaosMover/Source/ChaosMover/Private/ChaosMoverStateMachine.h b/Engine/Plugins/Experimental/ChaosMover/Source/ChaosMover/Private/ChaosMoverStateMachine.h
--- a/Engine/Plugins/Experimental/ChaosMover/Source/ChaosMover/Private/ChaosMoverStateMachine.h
+++ b/Engine/Plugins/Experimental/ChaosMover/Source/ChaosMover/Private/ChaosMoverStateMachine.h
@@ -7,20 +7,12 @@
 #include "InstantMovementEffect.h"
 #include "MovementMode.h"
 #include "Templates/SubclassOf.h"
+#include "ChaosMover/ChaosMoverSimulationTypes.h"
 
 class UChaosMoverSimulation;
 
 namespace UE::ChaosMover
 {
-	/** Mode change event structure, used to postpone callbacks to gameplay code when a mode has changed */
-	struct FMovementModeChangeEvent
-	{
-		FName PreviousMovementModeName;
-		TWeakObjectPtr<UBaseMovementMode> PreviousMovementMode;
-		FName NextMovementModeName;
-		TWeakObjectPtr<UBaseMovementMode> NextMovementMode;
-	};
-
 	class CHAOSMOVER_API FMoverStateMachine
 	{
 	public:
@@ -44,6 +36,7 @@ namespace UE::ChaosMover
 		void UnregisterGlobalTransition(TWeakObjectPtr<UBaseMovementModeTransition> Transition);
 		void ClearAllGlobalTransitions();
 
+		FName GetDefaultModeName() const;
 		void SetDefaultMode(FName NewDefaultModeName);
 
 		void QueueNextMode(FName DesiredNextModeName, bool bShouldReenter = false);
@@ -67,9 +60,6 @@ namespace UE::ChaosMover
 
 		void QueueInstantMovementEffect(TSharedPtr<FInstantMovementEffect> Effect);
 
-		const TArray<FMovementModeChangeEvent>& GetQueuedModeChangeEvents() const;
-		void ClearQueuedModeChangeEvents();
-
 		void SetOwnerActorName(const FString& InOwnerActorName);
 		void SetOwnerActorLocalNetRole(ENetRole InOwnerActorLocalNetRole);
 
@@ -95,9 +85,6 @@ namespace UE::ChaosMover
 		/** Effects that are queued to be applied to the simulation at the start of the next sim subtick or at the end of this tick */
 		TArray<TSharedPtr<FInstantMovementEffect>> QueuedInstantEffects;
 
-		/** This queue is used to record movement mode transitions and postpone the execution of their callbacks */
-		TArray<FMovementModeChangeEvent> QueuedModeChangeEvents;
-
 		TWeakObjectPtr<UChaosMoverSimulation> Simulation;
 
 	private:
diff --git a/Engine/Plugins/Experimental/ChaosMover/Source/ChaosMover/Private/Character/Modes/ChaosWalkingMode.cpp b/Engine/Plugins/Experimental/ChaosMover/Source/ChaosMover/Private/Character/Modes/ChaosWalkingMode.cpp
--- a/Engine/Plugins/Experimental/ChaosMover/Source/ChaosMover/Private/Character/Modes/ChaosWalkingMode.cpp
+++ b/Engine/Plugins/Experimental/ChaosMover/Source/ChaosMover/Private/Character/Modes/ChaosWalkingMode.cpp
@@ -415,7 +415,7 @@ void UChaosWalkingMode::ModifyContacts(const FMoverTimeStep& TimeStep, const UE:
 
 	// Get the updated (character) particle
 	Chaos::FGeometryParticleHandle* UpdatedParticle = nullptr;
-	const FChaosMoverSimulationDefaultInputs* SimInputs = InputData.LocalInputParams.AuxStateCollection.FindDataByType<FChaosMoverSimulationDefaultInputs>();
+	const FChaosMoverSimulationDefaultInputs* SimInputs = Simulation->GetLocalSimInput().FindDataByType<FChaosMoverSimulationDefaultInputs>();
 	if (SimInputs)
 	{
 		Chaos::FReadPhysicsObjectInterface_Internal ReadInterface = Chaos::FPhysicsObjectInternalInterface::GetRead();
diff --git a/Engine/Plugins/Experimental/ChaosMover/Source/ChaosMover/Public/ChaosMover/Backends/ChaosMoverBackend.h b/Engine/Plugins/Experimental/ChaosMover/Source/ChaosMover/Public/ChaosMover/Backends/ChaosMoverBackend.h
--- a/Engine/Plugins/Experimental/ChaosMover/Source/ChaosMover/Public/ChaosMover/Backends/ChaosMoverBackend.h
+++ b/Engine/Plugins/Experimental/ChaosMover/Source/ChaosMover/Public/ChaosMover/Backends/ChaosMoverBackend.h
@@ -18,6 +18,9 @@ namespace Chaos
 {
 	class FCharacterGroundConstraint;
 	class FCharacterGroundConstraintProxy;
+
+	class FJointConstraint;
+	class FJointConstraintPhysicsProxy;
 }
 
 UCLASS(Within = MoverComponent)
@@ -64,6 +67,10 @@ protected:
 	virtual void DeinitSimulation();
 	virtual void CreatePhysics();
 	virtual void DestroyPhysics();
+	void CreateCharacterGroundConstraint();
+	void DestroyCharacterGroundConstraint();
+	void CreatePathTargetConstraint();
+	void DestroyPathTargetConstraint();
 
 	virtual void GenerateInput(const FMoverTimeStep& TimeStep, UE::ChaosMover::FSimulationInputData& InputData);
 
@@ -74,6 +81,7 @@ protected:
 
 	UMoverComponent& GetMoverComponent() const;
 	Chaos::FPhysicsSolver* GetPhysicsSolver() const;
+	Chaos::FPhysicsObject* GetPhysicsObject() const;
 	Chaos::FPBDRigidParticle* GetControlledParticle() const;
 	FMoverTimeStep GetMoverTimeStep(double DeltaTimeInSeconds, int32 NetworkPhysicsTickOffset) const;
 
@@ -96,8 +104,15 @@ protected:
 	TObjectPtr<UChaosMoverSimulation> Simulation;
 
 	UE::Mover::FSimulationOutputRecord SimOutputRecord;
+
+	// Character ground constraint, for moving on ground like characters
 	TUniquePtr<Chaos::FCharacterGroundConstraint> CharacterGroundConstraint;
 
+	// Path target constraint, for moving along a path
+	FPhysicsConstraintHandle PathTargetConstraintHandle;
+	FConstraintInstance PathTargetConstraintInstance;
+	FPhysicsUserData PathTargetConstraintPhysicsUserData;
+
 	bool bIsUsingAsyncPhysics = false;
 	bool bWantsDestroySim = false;
 	bool bWantsCreateSim = false;
diff --git a/Engine/Plugins/Experimental/ChaosMover/Source/ChaosMover/Public/ChaosMover/ChaosMoverSimulation.h b/Engine/Plugins/Experimental/ChaosMover/Source/ChaosMover/Public/ChaosMover/ChaosMoverSimulation.h
--- a/Engine/Plugins/Experimental/ChaosMover/Source/ChaosMover/Public/ChaosMover/ChaosMoverSimulation.h
+++ b/Engine/Plugins/Experimental/ChaosMover/Source/ChaosMover/Public/ChaosMover/ChaosMoverSimulation.h
@@ -9,20 +9,17 @@
 #include "MoveLibrary/MoverBlackboard.h"
 #include "MoverSimulation.h"
 #include "UObject/WeakObjectPtrTemplates.h"
+#include "PhysicsEngine/ConstraintInstance.h"
 
 #include "ChaosMoverSimulation.generated.h"
 
 namespace Chaos
 {
 	class FCharacterGroundConstraintProxy;
+	class FJointConstraintPhysicsProxy;
 	class FCollisionContactModifier;
 }
 
-namespace UE::ChaosMover
-{
-	struct FMovementModeChangeEvent;
-}
-
 UCLASS(BlueprintType)
 class CHAOSMOVER_API UChaosMoverSimulation : public UMoverSimulation
 {
@@ -31,9 +28,14 @@ class CHAOSMOVER_API UChaosMoverSimulation : public UMoverSimulation
 public:
 	UChaosMoverSimulation();
 
+	virtual void OnMovementModeChanged(FName PreviousMovementModeName, TWeakObjectPtr<UBaseMovementMode> PreviousMovementMode, FName NextMovementModeName, TWeakObjectPtr<UBaseMovementMode> NextMovementMode);
+
+	// Returns the local simulation input MoverDataCollection, to read local non networked data passed to the simulation by the gameplay thread
 	UFUNCTION(BlueprintPure, Category = Mover)
 	const FMoverDataCollection& GetLocalSimInput() const;
 
+	// Returns the local simulation input MoverDataCollection, to pass local non networked data to the simulation
+	// Only available from the gameplay thread
 	UFUNCTION(BlueprintPure, Category = Mover)
 	FMoverDataCollection& GetLocalSimInput_Mutable();
 
@@ -60,6 +62,8 @@ public:
 		TWeakObjectPtr<UImmediateMovementModeTransition> ImmediateModeTransition = nullptr;
 		TWeakObjectPtr<UMovementMixer> MovementMixer = nullptr;
 		Chaos::FCharacterGroundConstraintProxy* CharacterConstraintProxy = nullptr;
+		Chaos::FJointConstraintPhysicsProxy* PathTargetConstraintProxy = nullptr;
+		Chaos::FSingleParticlePhysicsProxy* PathTargetKinematicEndPointProxy = nullptr;
 		Chaos::FConstPhysicsObjectHandle PhysicsObject = nullptr;
 		Chaos::FPhysicsSolver* Solver = nullptr;
 		UWorld* World = nullptr;
@@ -70,8 +74,7 @@ public:
 	void SimulationTick(const FMoverTimeStep& TimeStep, const UE::ChaosMover::FSimulationInputData& InputData, UE::ChaosMover::FSimulationOutputData& OutputData);
 	void ModifyContacts(const FMoverTimeStep& TimeStep, const UE::ChaosMover::FSimulationInputData& InputData, const UE::ChaosMover::FSimulationOutputData& OutputData, Chaos::FCollisionContactModifier& Modifier) const;
 
-	const TArray<UE::ChaosMover::FMovementModeChangeEvent>& GetQueuedModeChangeEvents() const;
-	void ClearQueuedModeChangeEvents();
+	void GetAndClearQueuedModeChangeEvents(TArray<UE::ChaosMover::FMovementModeChangeEvent>& OutModeChangedEvents);
 
 	void InitNetInputData(const FMoverInputCmdContext& InNetInputCmd);
 	void ApplyNetInputData(const FMoverInputCmdContext& InNetInputCmd);
@@ -79,6 +82,12 @@ public:
 	void ApplyNetStateData(const FMoverSyncState& InNetSyncState);
 	void BuildNetStateData(FMoverSyncState& OutNetSyncState) const;
 
+	//~ Debugging Util functions
+	// Collection for holding extra debug data, that will be sent to the Chaos Visual Debugger for debugging
+	UFUNCTION(BlueprintPure, Category = Mover)
+	FMoverDataCollection& GetDebugSimData();
+	//~ End of Debugging Util functions
+
 protected:
 	virtual void OnInit();
 	virtual void OnDeinit();
@@ -87,27 +96,73 @@ protected:
 	virtual void OnPostSimulationTick(const FMoverTimeStep& TimeStep, UE::ChaosMover::FSimulationOutputData& OutputData);
 	virtual void OnModifyContacts(const FMoverTimeStep& TimeStep, const UE::ChaosMover::FSimulationInputData& InputData, const UE::ChaosMover::FSimulationOutputData& OutputData, Chaos::FCollisionContactModifier& Modifier) const;
 
+	// Character-like movement sim steps
 	virtual void PreSimulationTickCharacter(const IChaosCharacterMovementModeInterface& CharacterMode, const FMoverTimeStep& TimeStep, const UE::ChaosMover::FSimulationInputData& InputData);
 	virtual void PostSimulationTickCharacter(const IChaosCharacterMovementModeInterface& CharacterMode, const FMoverTimeStep& TimeStep, UE::ChaosMover::FSimulationOutputData& OutputData);
 	virtual void PostSimulationTickCharacterConstraint(const IChaosCharacterConstraintMovementModeInterface& CharacterConstraintMode, const FMoverTimeStep& TimeStep, UE::ChaosMover::FSimulationOutputData& OutputData);
 
+	// Pathed movement sim steps (kinematic or constrained)
+	virtual void PostSimulationTickPathedMovement(const IChaosPathedMovementModeInterface& ConstraintMode, const FMoverTimeStep& TimeStep, UE::ChaosMover::FSimulationOutputData& OutputData);
+
 	Chaos::FPBDRigidParticleHandle* GetControlledParticle() const;
+
+	void EnqueueModeChangeEvent(const UE::ChaosMover::FMovementModeChangeEvent& MovementModeChangeEvent);
+
+	// Object State.
+	// In kinematic mode, the updated primitive component needs both its SimulatePhysics and UpdateKinematicFromSimulation options checked
+	void SetControlledParticleDynamic();
+	void SetControlledParticleKinematic();
+	bool IsControlledParticleDynamic() const;
+	bool IsControlledParticleKinematic() const;
+
+	// Character ground constraint
 	void EnableCharacterConstraint();
 	void DisableCharacterConstraint();
 
+	// Path target constraint
+	void EnablePathTargetConstraint();
+	void DisablePathTargetConstraint();
+
+	void TraceMoverData(const UE::ChaosMover::FSimulationOutputData& OutputData);
+
+	// State structs
 	FMoverSyncState CurrentSyncState;
+	// Data internal to the simulation
 	FMoverDataCollection InternalSimData;
+	// Local input data, usually sent by the gameplay side locally, that is not expected to differ from that on the server so doesn't warrant networking
 	FMoverDataCollection LocalSimInput;
+	// Debug Data collection, sent to Chaos Visual Debugger when Trace Extra Sim Debug Data is selected
+	FMoverDataCollection DebugSimData;
 
+	// Movement mode state machine
 	UE::ChaosMover::FMoverStateMachine StateMachine;
 
+	// Optional movement mixer
 	TWeakObjectPtr<UMovementMixer> MovementMixerWeakPtr = nullptr;
-	Chaos::FCharacterGroundConstraintProxy* CharacterConstraintProxy = nullptr;
+
+	// Controlled physics object
 	Chaos::FConstPhysicsObjectHandle PhysicsObject = nullptr;
+
+	// Character ground constraint
+	Chaos::FCharacterGroundConstraintProxy* CharacterConstraintProxy = nullptr;
+
+	// Path target constraint, for movement following a path in a physics based manner
+	Chaos::FJointConstraintPhysicsProxy* PathTargetConstraintProxy = nullptr;
+	Chaos::FSingleParticlePhysicsProxy* PathTargetKinematicEndPointProxy = nullptr;
+
 	Chaos::FPhysicsSolver* Solver = nullptr;
 	UWorld* World = nullptr;
 
 private:
 	FMoverInputCmdContext InputCmd;
+
+	// This queue is used to record movement mode transitions and postpone the execution of their callbacks
+	// This implementation is flawed, even with a mutex, we need to implement differently
+	TArray<UE::ChaosMover::FMovementModeChangeEvent> QueuedModeChangeEvents;
+	// FinalizeFrame and the next physics tick can be concurrent so QueuedModeChangeEvents access was not thread safe
+	// This mutex lock is just temporary to prevent crashes
+	// TODO: Implement a lockless solution (output the mode changed events as part of sim output instead)
+	mutable FRWLock ModeChangedEventsLock;
+
 	bool bInputCmdOverridden = false;
 };
diff --git a/Engine/Plugins/Experimental/ChaosMover/Source/ChaosMover/Public/ChaosMover/ChaosMoverSimulationTypes.h b/Engine/Plugins/Experimental/ChaosMover/Source/ChaosMover/Public/ChaosMover/ChaosMoverSimulationTypes.h
--- a/Engine/Plugins/Experimental/ChaosMover/Source/ChaosMover/Public/ChaosMover/ChaosMoverSimulationTypes.h
+++ b/Engine/Plugins/Experimental/ChaosMover/Source/ChaosMover/Public/ChaosMover/ChaosMoverSimulationTypes.h
@@ -10,10 +10,13 @@
 
 namespace Chaos
 {
+	class FPBDJointSettings;
 	class FCharacterGroundConstraintSettings;
 }
 
 struct FConstraintProfileProperties;
+class UChaosMoverSimulation;
+class UBaseMovementMode;
 
 namespace UE::ChaosMover
 {
@@ -31,6 +34,19 @@ namespace UE::ChaosMover
 
 	using FSimulationOutputData = ::UE::Mover::FSimulationOutputData;
 
+	/** Mode change event structure, used to postpone callbacks to gameplay code when a mode has changed */
+	struct FMovementModeChangeEvent
+	{
+		FName PreviousMovementModeName;
+		TWeakObjectPtr<UBaseMovementMode> PreviousMovementMode;
+		FName NextMovementModeName;
+		TWeakObjectPtr<UBaseMovementMode> NextMovementMode;
+	};
+
+	// Util function to be able to get the debug sim data collection from a UChaosMoverSimulation from another plugin,
+	// without including ChaosMoverSimulation.h
+	CHAOSMOVER_API FMoverDataCollection& GetDebugSimData(UChaosMoverSimulation* Simulation);
+
 } // namespace UE::ChaosMover
 
 // Default chaos mover character simulation state, contains state basic for character simulation
@@ -126,6 +142,7 @@ public:
 	virtual bool ShouldCharacterRemainUpright() const = 0;
 };
 
+// Interface for mover modes moving on ground like characters, using a character ground constraint
 UINTERFACE(MinimalAPI)
 class UChaosCharacterConstraintMovementModeInterface : public UInterface
 {
@@ -138,20 +155,25 @@ class IChaosCharacterConstraintMovementModeInterface
 
 public:
 	virtual bool ShouldEnableConstraint() const = 0;
-	virtual  void UpdateConstraintSettings(Chaos::FCharacterGroundConstraintSettings& ConstraintSettings) const = 0;
+	virtual void UpdateConstraintSettings(Chaos::FCharacterGroundConstraintSettings& ConstraintSettings) const = 0;
 };
 
+// Interface for mover modes moving on a path, with a constraint or kinematically
 UINTERFACE(MinimalAPI)
-class UChaosConstraintMovementModeInterface : public UInterface
+class UChaosPathedMovementModeInterface : public UInterface
 {
 	GENERATED_BODY()
 };
 
-class IChaosConstraintMovementModeInterface
+class IChaosPathedMovementModeInterface
 {
 	GENERATED_BODY()
 
 public:
-	virtual bool ShouldEnableConstraint() const = 0;
-	virtual const FConstraintProfileProperties& GetConstraintProperties() const = 0;
+
+	// Whether the simulation should use a constraint to move the actor along the path, or moving the actor kinematically
+	virtual bool ShouldUseConstraint() const = 0;
+
+	// Constraint settings of the joint constraint, only valid when ShouldUseConstraint() is true
+	virtual const Chaos::FPBDJointSettings& GetConstraintSettings() const = 0;
 };
diff --git a/Engine/Plugins/Experimental/Mover/Source/Mover/Private/Backends/MoverNetworkPhysicsLiaison.cpp b/Engine/Plugins/Experimental/Mover/Source/Mover/Private/Backends/MoverNetworkPhysicsLiaison.cpp
--- a/Engine/Plugins/Experimental/Mover/Source/Mover/Private/Backends/MoverNetworkPhysicsLiaison.cpp
+++ b/Engine/Plugins/Experimental/Mover/Source/Mover/Private/Backends/MoverNetworkPhysicsLiaison.cpp
@@ -301,7 +301,7 @@ void UMoverNetworkPhysicsLiaisonComponent::PerformProcessInputs_Internal(int32 P
 	}
 	
 	// This will only do something if CVD is actively tracing and the mover info CVD data channel checked
-	FMoverCVDRuntimeTrace::TraceMoverData(&GetMoverComponent(), &Input.InputCmd, &Input.SyncState);
+	UE::MoverUtils::FMoverCVDRuntimeTrace::TraceMoverData(&GetMoverComponent(), &Input.InputCmd, &Input.SyncState);
 }
 
 bool UMoverNetworkPhysicsLiaisonComponent::CanSimulate_Internal(const FPhysicsMoverSimulationTickParams& TickParams, const FPhysicsMoverAsyncInput& Input) const
diff --git a/Engine/Plugins/Experimental/Mover/Source/Mover/Private/ChaosVisualDebugger/MoverCVDRuntimeTrace.cpp b/Engine/Plugins/Experimental/Mover/Source/Mover/Private/ChaosVisualDebugger/MoverCVDRuntimeTrace.cpp
--- a/Engine/Plugins/Experimental/Mover/Source/Mover/Private/ChaosVisualDebugger/MoverCVDRuntimeTrace.cpp
+++ b/Engine/Plugins/Experimental/Mover/Source/Mover/Private/ChaosVisualDebugger/MoverCVDRuntimeTrace.cpp
@@ -19,6 +19,9 @@
 #include "Interfaces/IPhysicsComponent.h"
 #include "Engine/World.h"
 
+namespace UE::MoverUtils
+{
+
 CVD_DEFINE_OPTIONAL_DATA_CHANNEL(MoverNetworkedData, EChaosVDDataChannelInitializationFlags::CanChangeEnabledState)
 CVD_DEFINE_OPTIONAL_DATA_CHANNEL(MoverLocalSimData, EChaosVDDataChannelInitializationFlags::CanChangeEnabledState)
 
@@ -141,7 +144,35 @@ void FMoverCVDRuntimeTrace::WrapSimData(uint32 SolverID, uint32 ParticleID, cons
 	}
 }
 
-void FMoverCVDRuntimeTrace::TraceMoverData(UMoverComponent* MoverComponent, const FMoverInputCmdContext* InputCmd, const FMoverSyncState* SyncState, const FMoverDataCollection* LocalSimData)
+USTRUCT()
+struct FMoverMergeableDataCollection : FMoverDataCollection
+{
+	void AppendShallow(const FMoverDataCollection& Other)
+	{
+		const FMoverMergeableDataCollection& MergeableOther = reinterpret_cast<const FMoverMergeableDataCollection&>(Other);
+		for (int i = 0; i < MergeableOther.DataArray.Num(); ++i)
+		{
+			if (MergeableOther.DataArray[i].IsValid())
+			{
+				DataArray.Add(MergeableOther.DataArray[i]);
+			}
+		}
+	}
+};
+
+void CombineDataCollections(const NamedDataCollections& DataCollections, FMoverDataCollection& OutDataCollection)
+{
+	FMoverMergeableDataCollection& OutMergeableDataCollection = reinterpret_cast<FMoverMergeableDataCollection&>(OutDataCollection);
+	for (const TPair<FName, const FMoverDataCollection*>& DataCollectionEntry : DataCollections)
+	{
+		if (const FMoverDataCollection* DataCollectionToAppend = DataCollectionEntry.Value)
+		{
+			OutMergeableDataCollection.AppendShallow(*DataCollectionToAppend);
+		}
+	}
+}
+
+void FMoverCVDRuntimeTrace::TraceMoverData(UMoverComponent* MoverComponent, const FMoverInputCmdContext* InputCmd, const FMoverSyncState* SyncState, const NamedDataCollections* LocalSimDataCollections /*= nullptr*/)
 {
 	if (!FChaosVisualDebuggerTrace::IsTracing())
 	{
@@ -170,13 +201,19 @@ void FMoverCVDRuntimeTrace::TraceMoverData(UMoverComponent* MoverComponent, cons
 		}
 
 		int32 SolverID = CVD_TRACE_GET_SOLVER_ID_FROM_WORLD(World);
-		// LocalSimState could add a lot of extra bytes, especially without some sort of delta serialization, so it is only optionally recorded
-		const FMoverDataCollection* RecordedLocalSimData = CVDDC_MoverLocalSimData->IsChannelEnabled() ? LocalSimData : nullptr;
+		const FMoverDataCollection* RecordedLocalSimData = nullptr;
+		FMoverDataCollection MergedDataCollection;
+		if (CVDDC_MoverLocalSimData->IsChannelEnabled() && LocalSimDataCollections)
+		{
+			// LocalSimState could add a lot of extra bytes, especially without some sort of delta serialization, so it is only optionally recorded
+			CombineDataCollections(*LocalSimDataCollections, MergedDataCollection);
+			RecordedLocalSimData = &MergedDataCollection;
+		}
 		TraceMoverDataPrivate(SolverID, ParticleID, InputCmd, SyncState, RecordedLocalSimData);
 	}
 }
 
-void FMoverCVDRuntimeTrace::TraceMoverData(uint32 SolverID, uint32 ParticleID, const FMoverInputCmdContext* InputCmd, const FMoverSyncState* SyncState, const FMoverDataCollection* LocalSimData)
+void FMoverCVDRuntimeTrace::TraceMoverData(uint32 SolverID, uint32 ParticleID, const FMoverInputCmdContext* InputCmd, const FMoverSyncState* SyncState, const NamedDataCollections* LocalSimDataCollections /*= nullptr*/)
 {
 	if (!FChaosVisualDebuggerTrace::IsTracing())
 	{
@@ -193,10 +230,18 @@ void FMoverCVDRuntimeTrace::TraceMoverData(uint32 SolverID, uint32 ParticleID, c
 		return;
 	}
 		
-	TraceMoverDataPrivate(SolverID, ParticleID, InputCmd, SyncState, LocalSimData);
+	const FMoverDataCollection* RecordedLocalSimData = nullptr;
+	FMoverDataCollection MergedDataCollection;
+	if (CVDDC_MoverLocalSimData->IsChannelEnabled() && LocalSimDataCollections)
+	{
+		// LocalSimState could add a lot of extra bytes, especially without some sort of delta serialization, so it is only optionally recorded
+		CombineDataCollections(*LocalSimDataCollections, MergedDataCollection);
+		RecordedLocalSimData = &MergedDataCollection;
+	}
+	TraceMoverDataPrivate(SolverID, ParticleID, InputCmd, SyncState, RecordedLocalSimData);
 }
 
-void FMoverCVDRuntimeTrace::TraceMoverDataPrivate(uint32 SolverID, uint32 ParticleID, const FMoverInputCmdContext* InputCmd, const FMoverSyncState* SyncState, const FMoverDataCollection* LocalSimData)
+void FMoverCVDRuntimeTrace::TraceMoverDataPrivate(uint32 SolverID, uint32 ParticleID, const FMoverInputCmdContext* InputCmd, const FMoverSyncState* SyncState, const FMoverDataCollection* LocalSimData /*= nullptr*/)
 {
 	FMoverCVDSimDataWrapper SimDataWrapper;
 	WrapSimData(SolverID, ParticleID, *InputCmd, *SyncState, LocalSimData, SimDataWrapper);
@@ -208,4 +253,7 @@ void FMoverCVDRuntimeTrace::TraceMoverDataPrivate(uint32 SolverID, uint32 Partic
 	FChaosVisualDebuggerTrace::TraceBinaryData(TLSDataBuffer.BufferRef, FMoverCVDSimDataWrapper::WrapperTypeName);
 }
 
+} // namespace UE::MoverUtils
+
 #endif // WITH_CHAOS_VISUAL_DEBUGGER
+
diff --git a/Engine/Plugins/Experimental/Mover/Source/Mover/Private/MoverComponent.cpp b/Engine/Plugins/Experimental/Mover/Source/Mover/Private/MoverComponent.cpp
--- a/Engine/Plugins/Experimental/Mover/Source/Mover/Private/MoverComponent.cpp
+++ b/Engine/Plugins/Experimental/Mover/Source/Mover/Private/MoverComponent.cpp
@@ -464,7 +464,7 @@ void UMoverComponent::InitializeSimulationState(FMoverSyncState* OutSync, FMover
 void UMoverComponent::SimulationTick(const FMoverTimeStep& InTimeStep, const FMoverTickStartData& SimInput, OUT FMoverTickEndData& SimOutput)
 {
 	// Send mover info to the Chaos Visual Debugger (this will do nothing if CVD is not recording, or the mover info data channel not enabled)
-	FMoverCVDRuntimeTrace::TraceMoverData(this, &SimInput.InputCmd, &SimInput.SyncState);
+	UE::MoverUtils::FMoverCVDRuntimeTrace::TraceMoverData(this, &SimInput.InputCmd, &SimInput.SyncState);
 
 	const bool bIsResimulating = InTimeStep.BaseSimTimeMs <= CachedNewestSimTickTimeStep.BaseSimTimeMs;
 
@@ -855,7 +855,10 @@ void UMoverComponent::CreateDefaultInputAndState(FMoverInputCmdContext& OutInput
 	// Add all initial persistent sync state types
 	for (const FMoverDataPersistence& PersistentSyncEntry : PersistentSyncStateDataTypes)
 	{
-		OutSyncState.SyncStateCollection.FindOrAddDataByType(PersistentSyncEntry.RequiredType);
+		if (PersistentSyncEntry.RequiredType.Get()) // This can happen if a previously existing required type was removed, causing a crash
+		{
+			OutSyncState.SyncStateCollection.FindOrAddDataByType(PersistentSyncEntry.RequiredType);
+		}
 	}
 
 	// Mirror the scene component transform if we have one, otherwise it will be left at origin
diff --git a/Engine/Plugins/Experimental/Mover/Source/Mover/Private/PhysicsMover/PathedMovement/PathedMovementMode.cpp b/Engine/Plugins/Experimental/Mover/Source/Mover/Private/PhysicsMover/PathedMovement/PathedMovementMode.cpp
--- a/Engine/Plugins/Experimental/Mover/Source/Mover/Private/PhysicsMover/PathedMovement/PathedMovementMode.cpp
+++ b/Engine/Plugins/Experimental/Mover/Source/Mover/Private/PhysicsMover/PathedMovement/PathedMovementMode.cpp
@@ -89,12 +89,11 @@ void UPathedPhysicsMovementMode::OnProcessInput_Internal(int32 PhysicsStep, floa
 	{
 		CachedInputs = Inputs->Props;
 
-		const FPathedPhysicsMovementState& StartingState = Input.SyncState.SyncStateCollection.FindOrAddMutableDataByType<FPathedPhysicsMovementState>();
-		const FMutablePathedMovementProperties& CurProps = StartingState.MutableProps;
+		FPathedPhysicsMovementState& StartMoveState = Input.SyncState.SyncStateCollection.FindOrAddMutableDataByType<FPathedPhysicsMovementState>();
+		const FMutablePathedMovementProperties& CurProps = StartMoveState.MutableProps;
 		
 		if (CachedInputs.IsMoving() != CurProps.IsMoving())
 		{
-			FPathedPhysicsMovementState& StartMoveState = Input.SyncState.SyncStateCollection.FindOrAddMutableDataByType<FPathedPhysicsMovementState>();
 			const float PlaybackDuration = GetPlaybackDuration();
 			if (StartMoveState.MutableProps.bIsInReverse && StartMoveState.LastStopPlaybackTime == 0.f)
 			{
diff --git a/Engine/Plugins/Experimental/Mover/Source/Mover/Public/ChaosVisualDebugger/MoverCVDRuntimeTrace.h b/Engine/Plugins/Experimental/Mover/Source/Mover/Public/ChaosVisualDebugger/MoverCVDRuntimeTrace.h
--- a/Engine/Plugins/Experimental/Mover/Source/Mover/Public/ChaosVisualDebugger/MoverCVDRuntimeTrace.h
+++ b/Engine/Plugins/Experimental/Mover/Source/Mover/Public/ChaosVisualDebugger/MoverCVDRuntimeTrace.h
@@ -14,14 +14,20 @@ struct FMoverSyncState;
 struct FMoverDataCollection;
 class UMoverComponent;
 
+
+namespace UE::MoverUtils
+{
+
+typedef TArray<TPair<FName, const FMoverDataCollection*>> NamedDataCollections;
+
 #if WITH_CHAOS_VISUAL_DEBUGGER
 
 /** Utility functions used to trace Mover data into the Chaos Visual Debugger */
 class FMoverCVDRuntimeTrace
 {
 public:
-	MOVER_API static void TraceMoverData(UMoverComponent* MoverComponent, const FMoverInputCmdContext* InputCmd, const FMoverSyncState* SyncState, const FMoverDataCollection* LocalSimState = nullptr);
-	MOVER_API static void TraceMoverData(uint32 SolverID, uint32 ParticleID, const FMoverInputCmdContext* InputCmd, const FMoverSyncState* SyncState, const FMoverDataCollection* LocalSimState = nullptr);
+	MOVER_API static void TraceMoverData(UMoverComponent* MoverComponent, const FMoverInputCmdContext* InputCmd, const FMoverSyncState* SyncState, const NamedDataCollections* LocalSimDataCollections = nullptr);
+	MOVER_API static void TraceMoverData(uint32 SolverID, uint32 ParticleID, const FMoverInputCmdContext* InputCmd, const FMoverSyncState* SyncState, const NamedDataCollections* LocalSimDataCollections = nullptr);
 
 	MOVER_API static void UnwrapSimData(const FMoverCVDSimDataWrapper& InSimDataWrapper, TSharedPtr<FMoverInputCmdContext>& OutInputCmd, TSharedPtr<FMoverSyncState>& OutSyncState, TSharedPtr<FMoverDataCollection>& OutLocalSimState);
 	MOVER_API static void WrapSimData(uint32 SolverID, uint32 ParticleID, const FMoverInputCmdContext& InInputCmd, const FMoverSyncState& InSyncState, const FMoverDataCollection* LocalSimState, FMoverCVDSimDataWrapper& OutSimDataWrapper);
@@ -41,10 +47,12 @@ CVD_DECLARE_OPTIONAL_DATA_CHANNEL_EXTERN(MoverLocalSimData, MOVER_API)
 class FMoverCVDRuntimeTrace
 {
 public:
-	static void TraceMoverData(UMoverComponent* MoverComponent, const FMoverInputCmdContext* InputCmd, const FMoverSyncState* SyncState, const FMoverDataCollection* FMoverLocalSimState = nullptr) {}
-	static void TraceMoverData(uint32 SolverID, uint32 ParticleID, const FMoverInputCmdContext* InputCmd, const FMoverSyncState* SyncState, const FMoverDataCollection* FMoverLocalSimState = nullptr) {}
+	static void TraceMoverData(UMoverComponent* MoverComponent, const FMoverInputCmdContext* InputCmd, const FMoverSyncState* SyncState, const NamedDataCollections* FMoverLocalSimState = nullptr) {}
+	static void TraceMoverData(uint32 SolverID, uint32 ParticleID, const FMoverInputCmdContext* InputCmd, const FMoverSyncState* SyncState, const NamedDataCollections* FMoverLocalSimState = nullptr) {}
 	static void UnwrapSimData(const FMoverCVDSimDataWrapper& InSimDataWrapper, TSharedPtr<FMoverInputCmdContext>& OutInputCmd, TSharedPtr<FMoverSyncState>& OutSyncState, TSharedPtr<FMoverDataCollection>& FMoverLocalSimState) {}
 	static void WrapSimData(uint32 SolverID, uint32 ParticleID, const FMoverInputCmdContext& InInputCmd, const FMoverSyncState& InSyncState, const FMoverDataCollection* FMoverLocalSimState, FMoverCVDSimDataWrapper& OutSimDataWrapper) {}
 };
 
 #endif // WITH_CHAOS_VISUAL_DEBUGGER
+
+} // namespace UE::MoverUtils
\ No newline at end of file
diff --git a/Engine/Plugins/Experimental/Mover/Source/MoverCVDEditor/Private/MoverCVDSimDataComponent.cpp b/Engine/Plugins/Experimental/Mover/Source/MoverCVDEditor/Private/MoverCVDSimDataComponent.cpp
--- a/Engine/Plugins/Experimental/Mover/Source/MoverCVDEditor/Private/MoverCVDSimDataComponent.cpp
+++ b/Engine/Plugins/Experimental/Mover/Source/MoverCVDEditor/Private/MoverCVDSimDataComponent.cpp
@@ -60,7 +60,7 @@ bool UMoverCVDSimDataComponent::FindAndUnwrapSimDataForParticle(uint32 ParticleI
 		DeserializedMoverStates = MakeShared<FDeserializedMoverStates>();
 		// This means that the SimData wasn't deserialized yet, so we do it now
 		// Otherwise we use the cached version
-		FMoverCVDRuntimeTrace::UnwrapSimData(*SimDataPtr, DeserializedMoverStates->InputCommand, DeserializedMoverStates->SyncState, DeserializedMoverStates->LocalSimData);
+		UE::MoverUtils::FMoverCVDRuntimeTrace::UnwrapSimData(*SimDataPtr, DeserializedMoverStates->InputCommand, DeserializedMoverStates->SyncState, DeserializedMoverStates->LocalSimData);
 	}
 	if (DeserializedMoverStates)
 	{
