ChaosMover: ChaosMoverBackend to force physics replication to Resimulation and movement to replicate or movement will not replicate correctly.

diff --git a/Engine/Plugins/Experimental/ChaosMover/Source/ChaosMover/Private/Backends/ChaosMoverBackend.cpp b/Engine/Plugins/Experimental/ChaosMover/Source/ChaosMover/Private/Backends/ChaosMoverBackend.cpp
--- a/Engine/Plugins/Experimental/ChaosMover/Source/ChaosMover/Private/Backends/ChaosMoverBackend.cpp
+++ b/Engine/Plugins/Experimental/ChaosMover/Source/ChaosMover/Private/Backends/ChaosMoverBackend.cpp
@@ -39,6 +39,23 @@ UChaosMoverBackendComponent::UChaosMoverBackendComponent()
 	if (Chaos::FPhysicsSolverBase::IsNetworkPhysicsPredictionEnabled())
 	{
 		SetIsReplicatedByDefault(true);
+
+		// Let's make sure PhysicsReplicationMode is set to Resimulation and that movement is set to replicate
+		UWorld* World = GetWorld();
+		AActor* MyActor = GetOwner();
+		if (MyActor && World && World->IsGameWorld() && (World->GetNetMode() != ENetMode::NM_Standalone))
+		{
+			if (MyActor->GetPhysicsReplicationMode() != EPhysicsReplicationMode::Resimulation)
+			{
+				MyActor->SetPhysicsReplicationMode(EPhysicsReplicationMode::Resimulation);
+				UE_LOG(LogChaosMover, Log, TEXT("ChaosMoverBackend: Setting Physics Replication Mode to Resimulation for %s or movement will not replicate correctly"), *GetNameSafe(MyActor));
+			}
+			if (!MyActor->IsReplicatingMovement())
+			{
+				MyActor->SetReplicateMovement(true);
+				UE_LOG(LogChaosMover, Log, TEXT("ChaosMoverBackend: Turning ON Replicate Movement for %s or movement will not replicate correctly"), *GetNameSafe(MyActor));
+			}
+		}
 	}
 
 	Simulation = CreateDefaultSubobject<UChaosMoverSimulation>(TEXT("ChaosMoverSimulation"));
