Chaos Mover - Fixes to support different gravity directions

#rnx

diff --git a/Engine/Plugins/Experimental/ChaosMover/Source/ChaosMover/Private/ChaosMoverSimulation.cpp b/Engine/Plugins/Experimental/ChaosMover/Source/ChaosMover/Private/ChaosMoverSimulation.cpp
--- a/Engine/Plugins/Experimental/ChaosMover/Source/ChaosMover/Private/ChaosMoverSimulation.cpp
+++ b/Engine/Plugins/Experimental/ChaosMover/Source/ChaosMover/Private/ChaosMoverSimulation.cpp
@@ -305,29 +305,30 @@ void UChaosMoverSimulation::PostSimulationTickCharacter(const IChaosCharacterMov
 
 	if (Chaos::FPBDRigidParticleHandle* ParticleHandle = GetControlledParticle())
 	{
-		// Store the delta position and rotation locally
-		FRotator DeltaRotation = PostSimDefaultSyncState->GetOrientation_WorldSpace() - FRotator(ParticleHandle->GetR());
-		FRotator Winding, Remainder;
-		DeltaRotation.GetWindingAndRemainder(Winding, Remainder);
-
+		// Linear motion
 		CharacterSimState.TargetDeltaPosition = PostSimDefaultSyncState->GetLocation_WorldSpace() - ParticleHandle->GetX();
-		CharacterSimState.TargetDeltaFacing = FMath::DegreesToRadians(Remainder.Yaw);
-
-		// Set the particle velocity to hit the target sync state transform
 		ParticleHandle->SetV(PostSimDefaultSyncState->GetVelocity_WorldSpace());
-		const float DeltaTimeSeconds = TimeStep.StepMs * 0.001f;
-		if (DeltaTimeSeconds > UE_SMALL_NUMBER)
+
+		// Angular motion
+		FQuat TgtQuat = PostSimDefaultSyncState->GetOrientation_WorldSpace().Quaternion();
+		TgtQuat.EnforceShortestArcWith(ParticleHandle->GetR());
+		FQuat QuatRotation = TgtQuat * ParticleHandle->GetR().Inverse();
+		FVector AngularDisplacement = QuatRotation.ToRotationVector();
+
+		FVector UpDir = FVector::UpVector;
+		if (const FChaosMoverSimulationDefaultInputs* SimInputs = LocalSimInput.FindDataByType<FChaosMoverSimulationDefaultInputs>())
 		{
-			FVector AngularVelocity = (CharacterSimState.TargetDeltaFacing / DeltaTimeSeconds) * FVector::UpVector;
+			UpDir = SimInputs->UpDir;
+		}
+		CharacterSimState.TargetDeltaFacing = AngularDisplacement.Dot(UpDir);
 
-			if (!CharacterMode.ShouldCharacterRemainUpright())
+		if (CharacterMode.ShouldCharacterRemainUpright())
+		{
+			const float DeltaTimeSeconds = TimeStep.StepMs * 0.001f;
+			if (DeltaTimeSeconds > UE_SMALL_NUMBER)
 			{
-				FVector PreviousUp = ParticleHandle->GetR().RotateVector(FVector::UpVector).GetSafeNormal();
-				FVector TargetUp = PostSimDefaultSyncState->GetOrientation_WorldSpace().Quaternion().RotateVector(FVector::UpVector).GetSafeNormal();
-				AngularVelocity += (PreviousUp.Cross(TargetUp).GetSafeNormal() * FMath::Acos(TargetUp.Dot(PreviousUp))) / DeltaTimeSeconds;
+				ParticleHandle->SetW(AngularDisplacement / DeltaTimeSeconds);
 			}
-
-			ParticleHandle->SetW(AngularVelocity);
 		}
 	}
 	else
@@ -368,6 +369,13 @@ void UChaosMoverSimulation::PostSimulationTickCharacterConstraint(const IChaosCh
 		return;
 	}
 
+	// Update the up direction in the settings
+	Chaos::FCharacterGroundConstraintSettings& Settings = ConstraintHandle->GetSettings_Mutable();
+	if (const FChaosMoverSimulationDefaultInputs* SimInputs = LocalSimInput.FindDataByType<FChaosMoverSimulationDefaultInputs>())
+	{
+		Settings.VerticalAxis = SimInputs->UpDir;
+	}
+
 	check(Blackboard.Get());
 
 	FFloorCheckResult FloorResult;
diff --git a/Engine/Plugins/Experimental/ChaosMover/Source/ChaosMover/Private/Character/Modes/ChaosFallingMode.cpp b/Engine/Plugins/Experimental/ChaosMover/Source/ChaosMover/Private/Character/Modes/ChaosFallingMode.cpp
--- a/Engine/Plugins/Experimental/ChaosMover/Source/ChaosMover/Private/Character/Modes/ChaosFallingMode.cpp
+++ b/Engine/Plugins/Experimental/ChaosMover/Source/ChaosMover/Private/Character/Modes/ChaosFallingMode.cpp
@@ -200,7 +200,7 @@ void UChaosFallingMode::SimulationTick_Implementation(const FSimulationTickParam
 	}
 
 	// The physics simulation applies Z-only gravity acceleration via physics volumes, so we need to account for it here 
-	FVector TargetVel = ProposedMove.LinearVelocity - DefaultSimInputs->PhysicsObjectGravity * DefaultSimInputs->UpDir * DeltaSeconds;
+	FVector TargetVel = ProposedMove.LinearVelocity - DefaultSimInputs->PhysicsObjectGravity * FVector::UpVector * DeltaSeconds;
 	FVector TargetPos = StartingSyncState->GetLocation_WorldSpace() + TargetVel * DeltaSeconds;
 
 	OutputState.MovementEndState.RemainingMs = 0.0f;
diff --git a/Engine/Plugins/Experimental/ChaosMover/Source/ChaosMover/Private/Character/Modes/ChaosFlyingMode.cpp b/Engine/Plugins/Experimental/ChaosMover/Source/ChaosMover/Private/Character/Modes/ChaosFlyingMode.cpp
--- a/Engine/Plugins/Experimental/ChaosMover/Source/ChaosMover/Private/Character/Modes/ChaosFlyingMode.cpp
+++ b/Engine/Plugins/Experimental/ChaosMover/Source/ChaosMover/Private/Character/Modes/ChaosFlyingMode.cpp
@@ -127,7 +127,7 @@ void UChaosFlyingMode::SimulationTick_Implementation(const FSimulationTickParams
 	}
 
 	// The physics simulation applies Z-only gravity acceleration via physics volumes, so we need to account for it here 
-	FVector TargetVel = ProposedMove.LinearVelocity - DefaultSimInputs->PhysicsObjectGravity * DefaultSimInputs->UpDir * DeltaSeconds;
+	FVector TargetVel = ProposedMove.LinearVelocity - DefaultSimInputs->PhysicsObjectGravity * FVector::UpVector * DeltaSeconds;
 	FVector TargetPos = StartingSyncState->GetLocation_WorldSpace() + TargetVel * DeltaSeconds;
 
 	OutputState.MovementEndState.RemainingMs = 0.0f;
diff --git a/Engine/Plugins/Experimental/ChaosMover/Source/ChaosMover/Private/Character/Modes/ChaosWalkingMode.cpp b/Engine/Plugins/Experimental/ChaosMover/Source/ChaosMover/Private/Character/Modes/ChaosWalkingMode.cpp
--- a/Engine/Plugins/Experimental/ChaosMover/Source/ChaosMover/Private/Character/Modes/ChaosWalkingMode.cpp
+++ b/Engine/Plugins/Experimental/ChaosMover/Source/ChaosMover/Private/Character/Modes/ChaosWalkingMode.cpp
@@ -213,7 +213,7 @@ void UChaosWalkingMode::SimulationTick_Implementation(const FSimulationTickParam
 		// This is so that the gravity in this mode will be consistent with the gravity
 		// set on the mover, not the default physics gravity
 		const FVector ProjectedVelocity = StartingSyncState->GetVelocity_WorldSpace() + DefaultSimInputs->Gravity * DeltaSeconds;
-		FVector TargetVelocity = ProjectedVelocity - DefaultSimInputs->PhysicsObjectGravity * DefaultSimInputs->UpDir * DeltaSeconds;
+		FVector TargetVelocity = ProjectedVelocity - DefaultSimInputs->PhysicsObjectGravity * FVector::UpVector * DeltaSeconds;
 
 		// If we have movement intent and not moving straight up/down then use the proposed move plane velocity
 		// otherwise just fall with gravity
diff --git a/Engine/Source/Runtime/Experimental/Chaos/Private/Chaos/Character/CharacterGroundConstraintSolver.cpp b/Engine/Source/Runtime/Experimental/Chaos/Private/Chaos/Character/CharacterGroundConstraintSolver.cpp
--- a/Engine/Source/Runtime/Experimental/Chaos/Private/Chaos/Character/CharacterGroundConstraintSolver.cpp
+++ b/Engine/Source/Runtime/Experimental/Chaos/Private/Chaos/Character/CharacterGroundConstraintSolver.cpp
@@ -194,7 +194,7 @@ namespace Chaos
 			const Chaos::FMatrix33 CharacterInvInertia = Utilities::ComputeWorldSpaceInertia(Character_Q, CharacterInvILocal);
 			ConstraintData.CharacterInvI = FSolverMatrix33(CharacterInvInertia);
 			ConstraintData.VerticalAxis = FSolverVec3(Settings.VerticalAxis);
-			ConstraintData.CharacterVerticalAxis = FSolverVec3(Character_Q.RotateVector(Settings.VerticalAxis));
+			ConstraintData.CharacterVerticalAxis = FSolverVec3(Character_Q.RotateVector(FVector::UpVector));
 
 			//////////////////////////////////////////////////////////////////////////
 			// Motion Target
