Mover - Add initial implementation for simulation events and implement in Chaos Mover.

#rnx

diff --git a/Engine/Plugins/Experimental/ChaosMover/Source/ChaosMover/Private/Backends/ChaosMoverBackend.cpp b/Engine/Plugins/Experimental/ChaosMover/Source/ChaosMover/Private/Backends/ChaosMoverBackend.cpp
--- a/Engine/Plugins/Experimental/ChaosMover/Source/ChaosMover/Private/Backends/ChaosMoverBackend.cpp
+++ b/Engine/Plugins/Experimental/ChaosMover/Source/ChaosMover/Private/Backends/ChaosMoverBackend.cpp
@@ -371,21 +371,7 @@ void UChaosMoverBackendComponent::EndPlay(const EEndPlayReason::Type EndPlayReas
 	Super::EndPlay(EndPlayReason);
 }
 
-FMoverTimeStep UChaosMoverBackendComponent::GetMoverTimeStep(double DeltaTimeInSeconds, int32 NetworkPhysicsTickOffset) const
-{
-	FMoverTimeStep TimeStep;
-	if (Chaos::FPhysicsSolver* Solver = GetPhysicsSolver())
-	{
-		TimeStep.ServerFrame = Solver->GetCurrentFrame() + NetworkPhysicsTickOffset;
-		TimeStep.StepMs = DeltaTimeInSeconds * 1000.0f;
-		TimeStep.BaseSimTimeMs = Solver->GetPhysicsResultsTime_External() * 1000.0f + NetworkPhysicsTickOffset * TimeStep.StepMs;
-		TimeStep.bIsResimulating = false;
-	}
-
-	return TimeStep;
-}
-
-void UChaosMoverBackendComponent::ProduceInputData(int32 PhysicsStep, int32 NumSteps, double DeltaTimeInSeconds, int32 NetworkPhysicsTickOffset, UE::ChaosMover::FSimulationInputData& InputData)
+void UChaosMoverBackendComponent::ProduceInputData(int32 PhysicsStep, int32 NumSteps, const FMoverTimeStep& TimeStep, UE::ChaosMover::FSimulationInputData& InputData)
 {
 	Chaos::EnsureIsInGameThreadContext();
 
@@ -404,8 +390,6 @@ void UChaosMoverBackendComponent::ProduceInputData(int32 PhysicsStep, int32 NumS
 		bWantsCreateSim = false;
 	}
 
-	FMoverTimeStep TimeStep = GetMoverTimeStep(DeltaTimeInSeconds, NetworkPhysicsTickOffset);
-
 	if (!NetworkPhysicsComponent || NetworkPhysicsComponent->IsLocallyControlled())
 	{
 		GenerateInput(TimeStep, InputData);
@@ -467,7 +451,7 @@ void UChaosMoverBackendComponent::ConsumeOutputData(const FMoverTimeStep& TimeSt
 	SimOutputRecord.Add(TimeStep, OutputData);
 }
 
-void UChaosMoverBackendComponent::FinalizeFrame(double ResultsTimeInSeconds)
+void UChaosMoverBackendComponent::FinalizeFrame(float ResultsTimeInMs)
 {
 	Chaos::EnsureIsInGameThreadContext();
 
@@ -475,7 +459,7 @@ void UChaosMoverBackendComponent::FinalizeFrame(double ResultsTimeInSeconds)
 
 	FMoverTimeStep TimeStep;
 	UE::ChaosMover::FSimulationOutputData InterpolatedOutput;
-	SimOutputRecord.GetInterpolated(ResultsTimeInSeconds, TimeStep, InterpolatedOutput);
+	SimOutputRecord.GetInterpolated(ResultsTimeInMs, TimeStep, InterpolatedOutput);
 
 	// Physics interactions in the last frame may have caused a change in position or velocity that's different from what a simple lerp would predict,
 	// so stomp the lerped sync state's transform data with that of the actual particle after the last sim frame
@@ -488,10 +472,6 @@ void UChaosMoverBackendComponent::FinalizeFrame(double ResultsTimeInSeconds)
 		TransformSyncState.MoveDirectionIntent = TransformSyncState.GetCapturedMovementBaseQuat().UnrotateVector(TransformSyncState.MoveDirectionIntent);
 	}
 
-	// This eventually will process any events queued by the simulation during the async tick.
-	// For now it's just movement mode activation, deactivation and transition callbacks.
-	FlushModeChangeEvents();
-
 	MoverComp.SetSimulationOutput(TimeStep, InterpolatedOutput);
 	
 	if (MoverComp.OnPostSimulationTick.IsBound())
@@ -499,34 +479,3 @@ void UChaosMoverBackendComponent::FinalizeFrame(double ResultsTimeInSeconds)
 		MoverComp.OnPostSimulationTick.Broadcast(TimeStep);
 	}
 }
-
-/** Flushes the queue of mode changed events, which causes all movement mode change callbacks to be called
-* This is not async safe, should only be called on the gameplay thread
-*/
-void UChaosMoverBackendComponent::FlushModeChangeEvents()
-{
-	TArray<UE::ChaosMover::FMovementModeChangeEvent> QueuedModeChangeEvents;
-	Simulation->GetAndClearQueuedModeChangeEvents(QueuedModeChangeEvents);
-	UMoverComponent& MoverComp = GetMoverComponent();
-
-	for (int Idx = QueuedModeChangeEvents.Num() - 1; Idx >= 0; --Idx)
-	{
-		const UE::ChaosMover::FMovementModeChangeEvent& MovementModeChangeEvent = QueuedModeChangeEvents[Idx];
-		TStrongObjectPtr<UBaseMovementMode> PreviousMovementMode = MovementModeChangeEvent.PreviousMovementMode.Pin();
-		TStrongObjectPtr<UBaseMovementMode> NextMovementMode = MovementModeChangeEvent.NextMovementMode.Pin();
-
-		if (PreviousMovementMode)
-		{
-			PreviousMovementMode->Deactivate();
-		}
-
-		if (NextMovementMode)
-		{
-			NextMovementMode->Activate();
-		}
-
-		// Signal movement mode change event
-		MoverComp.OnMovementModeChanged.Broadcast(MovementModeChangeEvent.PreviousMovementModeName, MovementModeChangeEvent.NextMovementModeName);
-	}
-}
-
diff --git a/Engine/Plugins/Experimental/ChaosMover/Source/ChaosMover/Private/Backends/ChaosMoverSubsystem.cpp b/Engine/Plugins/Experimental/ChaosMover/Source/ChaosMover/Private/Backends/ChaosMoverSubsystem.cpp
--- a/Engine/Plugins/Experimental/ChaosMover/Source/ChaosMover/Private/Backends/ChaosMoverSubsystem.cpp
+++ b/Engine/Plugins/Experimental/ChaosMover/Source/ChaosMover/Private/Backends/ChaosMoverSubsystem.cpp
@@ -105,7 +105,8 @@ void UChaosMoverSubsystem::InjectInputs_External(int32 PhysicsStep, int32 NumSte
 	const int32 NetPhysicsTickOffset = GetNetworkPhysicsTickOffset();
 	AsyncInput->NetworkPhysicsTickOffset = NetPhysicsTickOffset;
 
-	double DeltaSeconds = 0.0;
+	// Compute the time step
+	FMoverTimeStep TimeStep;
 	if (UWorld* World = GetWorld())
 	{
 		if (FPhysScene* Scene = World->GetPhysicsScene())
@@ -116,18 +117,21 @@ void UChaosMoverSubsystem::InjectInputs_External(int32 PhysicsStep, int32 NumSte
 
 				if (Solver->IsUsingAsyncResults())
 				{
-					DeltaSeconds = Solver->GetAsyncDeltaTime();
+					TimeStep.StepMs = Solver->GetAsyncDeltaTime() * 1000.0f;
 				}
 				else
 				{
-					DeltaSeconds = World->GetDeltaSeconds();
+					TimeStep.StepMs = FMath::Clamp(World->GetDeltaSeconds(), Solver->GetMinDeltaTime_External(), Solver->GetMaxDeltaTime_External()) * 1000.0f;
 				}
+
+				TimeStep.ServerFrame = Solver->GetCurrentFrame() + NetPhysicsTickOffset;
+				TimeStep.BaseSimTimeMs = Solver->GetPhysicsResultsTime_External() * 1000.0f + NetPhysicsTickOffset * TimeStep.StepMs;
 			}
 		}
 	}
 
-	auto LambdaParallelUpdate = [NetPhysicsTickOffset, PhysicsStep, NumSteps, AsyncInput, DeltaSeconds, this](int32 Idx) {
-		Backends[Idx]->ProduceInputData(PhysicsStep, NumSteps, DeltaSeconds, NetPhysicsTickOffset, AsyncInput->InputData[Idx]);
+	auto LambdaParallelUpdate = [NetPhysicsTickOffset, PhysicsStep, NumSteps, AsyncInput, &TimeStep, this](int32 Idx) {
+		Backends[Idx]->ProduceInputData(PhysicsStep, NumSteps, TimeStep, AsyncInput->InputData[Idx]);
 	};
 
 	const bool ForceSingleThread = UE::ChaosMover::CVars::bForceSingleThreadedGT;
@@ -144,6 +148,8 @@ void UChaosMoverSubsystem::OnPostPhysicsTick(FChaosScene* Scene)
 	const bool ForceSingleThread = UE::ChaosMover::CVars::bForceSingleThreadedGT;
 
 	// Pop each async output
+	float LastDtInMs = 0.0f;
+
 	while (Chaos::TSimCallbackOutputHandle<UE::ChaosMover::FAsyncCallbackOutput> AsyncOutput = AsyncCallback->PopFutureOutputData_External())
 	{
 		auto LambdaParallelUpdate = [&AsyncOutput](int32 Idx) {
@@ -153,23 +159,29 @@ void UChaosMoverSubsystem::OnPostPhysicsTick(FChaosScene* Scene)
 			}
 		};
 
+		if (AsyncOutput->TimeStep.Num() > 0)
+		{
+			LastDtInMs = AsyncOutput->TimeStep[0].StepMs;
+		}
+
 		Chaos::PhysicsParallelFor(AsyncOutput->Backends.Num(), LambdaParallelUpdate, ForceSingleThread);
 	}
 
-	double ResultsTime = 0.0;
-	if (Scene != nullptr)
+	float ResultsTimeInMs = 0.0f;
+	if (Scene)
 	{
-		if (const Chaos::FPhysicsSolver* Solver = Scene->GetSolver())
+		if (Chaos::FPhysicsSolver* Solver = Scene->GetSolver())
 		{
-			ResultsTime = Solver->GetPhysicsResultsTime_External();
+			// Want the base time at the start of the frame
+			ResultsTimeInMs = Solver->GetPhysicsResultsTime_External() * 1000.0f - LastDtInMs;
 		}
 	}
 
 	// Call finalize frame
-	auto LambdaParallelUpdate = [this, &ResultsTime](int32 Idx) {
+	auto LambdaParallelUpdate = [this, &ResultsTimeInMs](int32 Idx) {
 		if (UChaosMoverBackendComponent* Backend = Backends[Idx].Get())
 		{
-			Backend->FinalizeFrame(ResultsTime);
+			Backend->FinalizeFrame(ResultsTimeInMs);
 		}
 	};
 
diff --git a/Engine/Plugins/Experimental/ChaosMover/Source/ChaosMover/Private/ChaosMoverSimulation.cpp b/Engine/Plugins/Experimental/ChaosMover/Source/ChaosMover/Private/ChaosMoverSimulation.cpp
--- a/Engine/Plugins/Experimental/ChaosMover/Source/ChaosMover/Private/ChaosMoverSimulation.cpp
+++ b/Engine/Plugins/Experimental/ChaosMover/Source/ChaosMover/Private/ChaosMoverSimulation.cpp
@@ -144,6 +144,8 @@ void UChaosMoverSimulation::Init(const FInitParams& InitParams)
 		StateMachine.RegisterGlobalTransition(Transition);
 	}
 
+	StateMachine.SetModeImmediately(StateMachine.GetDefaultModeName());
+
 	OnInit();
 }
 
@@ -527,6 +529,10 @@ void UChaosMoverSimulation::OnPostSimulationTick(const FMoverTimeStep& TimeStep,
 
 	CurrentSyncState = OutputData.SyncState;
 
+	// Extract the events into the output data and clear
+	OutputData.Events = Events;
+	Events.Empty();
+
 	// Send Debug Data to the Chaos Visual Debugger
 	TraceMoverData(OutputData);
 }
@@ -563,62 +569,63 @@ void UChaosMoverSimulation::OnModifyContacts(const FMoverTimeStep& TimeStep, con
 {
 }
 
-void UChaosMoverSimulation::OnMovementModeChanged(FName PreviousMovementModeName, TWeakObjectPtr<UBaseMovementMode> PreviousMovementMode, FName NextMovementModeName, TWeakObjectPtr<UBaseMovementMode> NextMovementMode)
+void UChaosMoverSimulation::AddEvent(TSharedPtr<FMoverSimulationEventData> Event)
 {
-	// Queue a mode changed event so it can be processed on the game thread to notify gameplay
-	EnqueueModeChangeEvent(/*FMovementModeChangeEvent*/{ PreviousMovementModeName, PreviousMovementMode, NextMovementModeName, NextMovementMode });
+	// Events are added to the event list for later extraction to game thread
+	// We also allow the simulation to react to the event immediately
+	Events.Add(Event);
 
-	TStrongObjectPtr<UBaseMovementMode> NextModePtr = NextMovementMode.Pin();
-	TStrongObjectPtr<UBaseMovementMode> PreviousModePtr = PreviousMovementMode.Pin();
+	if (const FMoverSimulationEventData* EventData = Event.Get())
+	{
+		ProcessSimulationEvent(*EventData);
+	}
+}
 
-	if (IChaosCharacterConstraintMovementModeInterface* NextCharacterConstraintMode = Cast<IChaosCharacterConstraintMovementModeInterface>(NextModePtr.Get()))
+void UChaosMoverSimulation::ProcessSimulationEvent(const FMoverSimulationEventData& EventData)
+{
+	if (const FMovementModeChangedEventData* ModeChangedEventData = EventData.CastTo<FMovementModeChangedEventData>())
 	{
-		bool IsCharacterConstraintInitialized = (CharacterConstraintProxy != nullptr) && CharacterConstraintProxy->IsInitialized();
-		Chaos::FCharacterGroundConstraintHandle* ConstraintHandle = IsCharacterConstraintInitialized ? CharacterConstraintProxy->GetPhysicsThreadAPI() : nullptr;
-		if (ConstraintHandle)
-		{
-			Chaos::FCharacterGroundConstraintSettings& Settings = ConstraintHandle->GetSettings_Mutable();
-			NextCharacterConstraintMode->UpdateConstraintSettings(Settings);
+		OnMovementModeChanged(*ModeChangedEventData);
+	}
+}
+
+void UChaosMoverSimulation::OnMovementModeChanged(const FMovementModeChangedEventData& ModeChangedData)
+{
+	TStrongObjectPtr<UBaseMovementMode> PreviousModePtr = StateMachine.FindMovementMode(ModeChangedData.PreviousModeName).Pin();
+	TStrongObjectPtr<UBaseMovementMode> NextModePtr = StateMachine.FindMovementMode(ModeChangedData.NewModeName).Pin();
 
-			// Character ground constraint modes currently assume moving a dynamic particle and using a character ground constraint
-			// Revise if we start supporting moving a character kinematically
-			if (!IsControlledParticleDynamic())
+	if (PreviousModePtr && NextModePtr)
+	{
+		if (IChaosCharacterConstraintMovementModeInterface* NextCharacterConstraintMode = Cast<IChaosCharacterConstraintMovementModeInterface>(NextModePtr.Get()))
+		{
+			bool IsCharacterConstraintInitialized = (CharacterConstraintProxy != nullptr) && CharacterConstraintProxy->IsInitialized();
+			Chaos::FCharacterGroundConstraintHandle* ConstraintHandle = IsCharacterConstraintInitialized ? CharacterConstraintProxy->GetPhysicsThreadAPI() : nullptr;
+			if (ConstraintHandle)
 			{
-				SetControlledParticleDynamic();
+				Chaos::FCharacterGroundConstraintSettings& Settings = ConstraintHandle->GetSettings_Mutable();
+				NextCharacterConstraintMode->UpdateConstraintSettings(Settings);
+
+				// Character ground constraint modes currently assume moving a dynamic particle and using a character ground constraint
+				// Revise if we start supporting moving a character kinematically
+				if (!IsControlledParticleDynamic())
+				{
+					SetControlledParticleDynamic();
+				}
 			}
 		}
-	}
 
-	if (IChaosPathedMovementModeInterface* NextPathTargetConstraintMode = Cast<IChaosPathedMovementModeInterface>(NextModePtr.Get()))
-	{
-		bool IsPathTargetConstraintInitialized = (PathTargetConstraintProxy != nullptr) && PathTargetConstraintProxy->IsInitialized();
-		Chaos::FPBDJointConstraintHandle* ConstraintHandle = IsPathTargetConstraintInitialized ? PathTargetConstraintProxy->GetHandle() : nullptr;
-		if (ConstraintHandle)
+		if (IChaosPathedMovementModeInterface* NextPathTargetConstraintMode = Cast<IChaosPathedMovementModeInterface>(NextModePtr.Get()))
 		{
-			ConstraintHandle->SetSettings(NextPathTargetConstraintMode->GetConstraintSettings());
+			bool IsPathTargetConstraintInitialized = (PathTargetConstraintProxy != nullptr) && PathTargetConstraintProxy->IsInitialized();
+			Chaos::FPBDJointConstraintHandle* ConstraintHandle = IsPathTargetConstraintInitialized ? PathTargetConstraintProxy->GetHandle() : nullptr;
+			if (ConstraintHandle)
+			{
+				ConstraintHandle->SetSettings(NextPathTargetConstraintMode->GetConstraintSettings());
+			}
 		}
 	}
 }
 
-void UChaosMoverSimulation::EnqueueModeChangeEvent(const UE::ChaosMover::FMovementModeChangeEvent& MovementModeChangeEvent)
-{
-	// FinalizeFrame and the next physics tick can be concurrent so this was not thread safe
-	// This mutex lock is just temporary to prevent crashes
-	// TODO: Implement a lockless solution (output the mode changed events as part of sim output instead)
-	FRWScopeLock Lock(ModeChangedEventsLock, SLT_Write);
-	QueuedModeChangeEvents.Add(MovementModeChangeEvent);
-}
-
-void UChaosMoverSimulation::GetAndClearQueuedModeChangeEvents(TArray<UE::ChaosMover::FMovementModeChangeEvent>& OutModeChangedEvents)
-{
-	// FinalizeFrame and the next physics tick can be concurrent so this was not thread safe
-	// This mutex lock is just temporary to prevent crashes
-	// TODO: Implement a lockless solution (output the mode changed events as part of sim output instead)
-	FRWScopeLock Lock(ModeChangedEventsLock, SLT_ReadOnly);
-	OutModeChangedEvents = QueuedModeChangeEvents;
-	QueuedModeChangeEvents.Empty();
-}
-
 Chaos::FPBDRigidParticleHandle* UChaosMoverSimulation::GetControlledParticle() const
 {
 	if (PhysicsObject)
diff --git a/Engine/Plugins/Experimental/ChaosMover/Source/ChaosMover/Private/ChaosMoverStateMachine.cpp b/Engine/Plugins/Experimental/ChaosMover/Source/ChaosMover/Private/ChaosMoverStateMachine.cpp
--- a/Engine/Plugins/Experimental/ChaosMover/Source/ChaosMover/Private/ChaosMoverStateMachine.cpp
+++ b/Engine/Plugins/Experimental/ChaosMover/Source/ChaosMover/Private/ChaosMoverStateMachine.cpp
@@ -174,6 +174,7 @@ namespace UE::ChaosMover
 		}
 
 		const float EndingSimTimeMs = TimeStep.BaseSimTimeMs + TimeStep.StepMs;
+		InternalSimTimeMs = TimeStep.BaseSimTimeMs;
 
 		//GEngine->AddOnScreenDebugMessage(-1, -0.1f, FColor::White, FString::Printf(TEXT("Mode Tick: %s  Queued: %s"), *CurrentModeName.ToString(), *NextModeName.ToString()));
 
@@ -194,6 +195,8 @@ namespace UE::ChaosMover
 		int32 NumConsecutiveFullRefundedSubsteps = 0;
 		while (SubTimeStep.BaseSimTimeMs < EndingSimTimeMs)
 		{
+			InternalSimTimeMs = SubTimeStep.BaseSimTimeMs;
+
 			SubstepStartSyncState = SubstepStartData.SyncState.SyncStateCollection.FindMutableDataByType<FMoverDefaultSyncState>();
 			FMoverDefaultSyncState* OutputSyncState = &OutputState.SyncState.SyncStateCollection.FindOrAddMutableDataByType<FMoverDefaultSyncState>();
 			OutputState.SyncState.MovementMode = CurrentModeName;
@@ -411,6 +414,8 @@ namespace UE::ChaosMover
 			++SubStepCount;
 		}
 
+		InternalSimTimeMs = EndingSimTimeMs;
+
 		// Apply any instant effects that were queued up during this tick and didn't get handled in a substep
 		if (ApplyInstantEffects(SubstepStartData, SubTimeStep, OutputState.SyncState))
 		{
@@ -489,19 +494,23 @@ namespace UE::ChaosMover
 					const bool bShouldNextModeReenter = QueuedModeTransition->ShouldReenter();
 					if ((CurrentModeName != NextModeName) || bShouldNextModeReenter)
 					{
-						const FName PreviousModeName = CurrentModeName;
-						TWeakObjectPtr<UBaseMovementMode>* FoundPreviousMovementMode = Modes.Find(PreviousModeName);
+						UE_LOG(LogChaosMover, Verbose, TEXT("AdvanceToNextMode: %s (%s) from %s to %s"),
+							*OwnerActorName, *UEnum::GetValueAsString(OwnerActorLocalNetRole), *CurrentModeName.ToString(), *NextModeName.ToString());
 
+						const FName PreviousModeName = CurrentModeName;
 						CurrentModeName = NextModeName;
 
-						UE_LOG(LogChaosMover, Verbose, TEXT("AdvanceToNextMode: %s (%s) from %s to %s"),
-							*OwnerActorName, *UEnum::GetValueAsString(OwnerActorLocalNetRole), *CurrentModeName.ToString(), *NextModeName.ToString());
+						if (PreviousModeName != NAME_None && Modes.Contains(PreviousModeName))
+						{
+							Modes[PreviousModeName]->Deactivate();
+						}
+
+						Modes[CurrentModeName]->Activate();
 
 						// Notify the simulation of a mode change so it can react accordingly
-						TWeakObjectPtr<UBaseMovementMode> PreviousMovementMode = FoundPreviousMovementMode ? *FoundPreviousMovementMode : TWeakObjectPtr<UBaseMovementMode>();
 						if (TStrongObjectPtr<UChaosMoverSimulation> SimStrongObjPtr = Simulation.Pin())
 						{
-							SimStrongObjPtr->OnMovementModeChanged(PreviousModeName, PreviousMovementMode, CurrentModeName, *FoundNextMovementMode);
+							SimStrongObjPtr->AddEvent(MakeShared<FMovementModeChangedEventData>(InternalSimTimeMs, PreviousModeName, NextModeName));
 						}
 					}
 				}
diff --git a/Engine/Plugins/Experimental/ChaosMover/Source/ChaosMover/Private/ChaosMoverStateMachine.h b/Engine/Plugins/Experimental/ChaosMover/Source/ChaosMover/Private/ChaosMoverStateMachine.h
--- a/Engine/Plugins/Experimental/ChaosMover/Source/ChaosMover/Private/ChaosMoverStateMachine.h
+++ b/Engine/Plugins/Experimental/ChaosMover/Source/ChaosMover/Private/ChaosMoverStateMachine.h
@@ -92,5 +92,7 @@ namespace UE::ChaosMover
 		void AdvanceToNextMode();
 		void FlushQueuedMovesToGroup(FLayeredMoveGroup& Group);
 		bool ApplyInstantEffects(const FMoverTickStartData& SubstepStartData, const FMoverTimeStep& SubTimeStep, FMoverSyncState& OutputState);
+	
+		float InternalSimTimeMs = 0.0f;
 	};
 }
\ No newline at end of file
diff --git a/Engine/Plugins/Experimental/ChaosMover/Source/ChaosMover/Private/Character/ChaosCharacterMoverComponent.cpp b/Engine/Plugins/Experimental/ChaosMover/Source/ChaosMover/Private/Character/ChaosCharacterMoverComponent.cpp
new file mode 100644
--- /dev/null
+++ b/Engine/Plugins/Experimental/ChaosMover/Source/ChaosMover/Private/Character/ChaosCharacterMoverComponent.cpp
@@ -0,0 +1,59 @@
+// Copyright Epic Games, Inc. All Rights Reserved.
+
+#include "ChaosMover/Character/ChaosCharacterMoverComponent.h"
+
+#include "ChaosMover/Backends/ChaosMoverBackend.h"
+#include "ChaosMover/Character/Modes/ChaosFallingMode.h"
+#include "ChaosMover/Character/Modes/ChaosFlyingMode.h"
+#include "ChaosMover/Character/Modes/ChaosWalkingMode.h"
+#include "ChaosMover/ChaosMoverLog.h"
+#include "DefaultMovementSet/CharacterMoverSimulationTypes.h"
+
+UChaosCharacterMoverComponent::UChaosCharacterMoverComponent()
+{
+	// Default movement modes
+	MovementModes.Add(DefaultModeNames::Walking, CreateDefaultSubobject<UChaosWalkingMode>(TEXT("DefaultChaosWalkingMode")));
+	MovementModes.Add(DefaultModeNames::Falling, CreateDefaultSubobject<UChaosFallingMode>(TEXT("DefaultChaosFallingMode")));
+	MovementModes.Add(DefaultModeNames::Flying, CreateDefaultSubobject<UChaosFlyingMode>(TEXT("DefaultChaosFlyingMode")));
+
+	StartingMovementMode = DefaultModeNames::Falling;
+
+	bHandleJump = false;
+	bHandleStanceChanges = false;
+
+	BackendClass = UChaosMoverBackendComponent::StaticClass();
+}
+
+void UChaosCharacterMoverComponent::ProcessSimulationEvent(const FMoverSimulationEventData& EventData)
+{
+	Super::ProcessSimulationEvent(EventData);
+
+	if (const FLandedEventData* LandedData = EventData.CastTo<FLandedEventData>())
+	{
+		OnLanded.Broadcast(LandedData->NewModeName, LandedData->HitResult);
+	}
+}
+
+void UChaosCharacterMoverComponent::SetAdditionalSimulationOutput(const FMoverDataCollection& Data)
+{
+	Super::SetAdditionalSimulationOutput(Data);
+
+	if (const FFloorResultData* FloorData = Data.FindDataByType<FFloorResultData>())
+	{
+		bFloorResultSet = true;
+		LatestFloorResult = FloorData->FloorResult;
+	}
+}
+
+bool UChaosCharacterMoverComponent::TryGetFloorCheckHitResult(FHitResult& OutHitResult) const
+{
+	if (bFloorResultSet)
+	{
+		OutHitResult = LatestFloorResult.HitResult;
+		return true;
+	}
+	else
+	{
+		return Super::TryGetFloorCheckHitResult(OutHitResult);
+	}
+}
\ No newline at end of file
diff --git a/Engine/Plugins/Experimental/ChaosMover/Source/ChaosMover/Private/Character/Transitions/ChaosCharacterLandingCheck.cpp b/Engine/Plugins/Experimental/ChaosMover/Source/ChaosMover/Private/Character/Transitions/ChaosCharacterLandingCheck.cpp
--- a/Engine/Plugins/Experimental/ChaosMover/Source/ChaosMover/Private/Character/Transitions/ChaosCharacterLandingCheck.cpp
+++ b/Engine/Plugins/Experimental/ChaosMover/Source/ChaosMover/Private/Character/Transitions/ChaosCharacterLandingCheck.cpp
@@ -6,6 +6,7 @@
 #include "ChaosMover/ChaosMoverSimulation.h"
 #include "ChaosMover/ChaosMoverSimulationTypes.h"
 #include "ChaosMover/Utilities/ChaosGroundMovementUtils.h"
+#include "DefaultMovementSet/CharacterMoverSimulationTypes.h"
 #include "MoveLibrary/FloorQueryUtils.h"
 #include "MoveLibrary/WaterMovementUtils.h"
 #include "MoverSimulationTypes.h"
@@ -16,6 +17,7 @@ UChaosCharacterLandingCheck::UChaosCharacterLandingCheck(const FObjectInitialize
 	: Super(ObjectInitializer)
 {
 	bSupportsAsync = true;
+	bFirstSubStepOnly = true;
 
 	TransitionToGroundMode = DefaultModeNames::Walking;
 }
@@ -85,4 +87,16 @@ FTransitionEvalResult UChaosCharacterLandingCheck::Evaluate_Implementation(const
 
 void UChaosCharacterLandingCheck::Trigger_Implementation(const FSimulationTickParams& Params)
 {
+	// Add a landed event to the simulation (will be broadcast on GT during post sim)
+	if (UChaosMoverSimulation* Sim = GetSimulation_Mutable())
+	{
+		if (const UMoverBlackboard* Blackboard = Sim->GetBlackboard())
+		{
+			FFloorCheckResult FloorResult;
+			if (Blackboard->TryGet(CommonBlackboard::LastFloorResult, FloorResult))
+			{
+				Sim->AddEvent(MakeShared<FLandedEventData>(Params.TimeStep.BaseSimTimeMs, FloorResult.HitResult, TransitionToGroundMode));
+			}
+		}
+	}
 }
\ No newline at end of file
diff --git a/Engine/Plugins/Experimental/ChaosMover/Source/ChaosMover/Public/ChaosMover/Backends/ChaosMoverBackend.h b/Engine/Plugins/Experimental/ChaosMover/Source/ChaosMover/Public/ChaosMover/Backends/ChaosMoverBackend.h
--- a/Engine/Plugins/Experimental/ChaosMover/Source/ChaosMover/Public/ChaosMover/Backends/ChaosMoverBackend.h
+++ b/Engine/Plugins/Experimental/ChaosMover/Source/ChaosMover/Public/ChaosMover/Backends/ChaosMoverBackend.h
@@ -23,34 +23,37 @@ namespace Chaos
 	class FJointConstraintPhysicsProxy;
 }
 
-UCLASS(Within = MoverComponent)
-class CHAOSMOVER_API UChaosMoverBackendComponent : public UActorComponent, public IMoverBackendLiaisonInterface
+UCLASS(MinimalAPI, Within = MoverComponent)
+class UChaosMoverBackendComponent : public UActorComponent, public IMoverBackendLiaisonInterface
 {
 	GENERATED_BODY()
 
 public:
-	UChaosMoverBackendComponent();
+	CHAOSMOVER_API UChaosMoverBackendComponent();
 
-	virtual void InitializeComponent() override;
-	virtual void UninitializeComponent() override;
-	virtual void BeginPlay() override;
-	virtual void EndPlay(const EEndPlayReason::Type EndPlayReason) override;
+	CHAOSMOVER_API virtual void InitializeComponent() override;
+	CHAOSMOVER_API virtual void UninitializeComponent() override;
+	CHAOSMOVER_API virtual void BeginPlay() override;
+	CHAOSMOVER_API virtual void EndPlay(const EEndPlayReason::Type EndPlayReason) override;
 
-	virtual bool IsAsync() const override { return true; }
+	CHAOSMOVER_API virtual bool IsAsync() const override
+	{
+		return true;
+	}
 
 	// Unused IMoverBackendLiaisonInterface API
-	virtual float GetCurrentSimTimeMs()
+	CHAOSMOVER_API virtual float GetCurrentSimTimeMs()
 	{
 		return 0.0f;
 	}
-	virtual int32 GetCurrentSimFrame()
+	CHAOSMOVER_API virtual int32 GetCurrentSimFrame()
 	{
 		return 0;
 	}
 
-	virtual void ProduceInputData(int32 PhysicsStep, int32 NumSteps, double DeltaTimeInSeconds, int32 NetworkPhysicsTickOffset, UE::ChaosMover::FSimulationInputData& InputData);
-	virtual void ConsumeOutputData(const FMoverTimeStep& TimeStep, const UE::ChaosMover::FSimulationOutputData& OutputData);
-	virtual void FinalizeFrame(double ResultsTimeInSeconds);
+	CHAOSMOVER_API virtual void ProduceInputData(int32 PhysicsStep, int32 NumSteps, const FMoverTimeStep& TimeStep, UE::ChaosMover::FSimulationInputData& InputData);
+	CHAOSMOVER_API virtual void ConsumeOutputData(const FMoverTimeStep& TimeStep, const UE::ChaosMover::FSimulationOutputData& OutputData);
+	CHAOSMOVER_API virtual void FinalizeFrame(float ResultsTimeInMs);
 
 	UChaosMoverSimulation* GetSimulation()
 	{
@@ -63,33 +66,27 @@ public:
 	}
 
 protected:
-	virtual void InitSimulation();
-	virtual void DeinitSimulation();
-	virtual void CreatePhysics();
-	virtual void DestroyPhysics();
-	void CreateCharacterGroundConstraint();
-	void DestroyCharacterGroundConstraint();
-	void CreatePathTargetConstraint();
-	void DestroyPathTargetConstraint();
-
-	virtual void GenerateInput(const FMoverTimeStep& TimeStep, UE::ChaosMover::FSimulationInputData& InputData);
-
-	/**
-	 * Flushes the queue of movement mode change events, which causes all movement mode change callbacks to be called
-	 */
-	void FlushModeChangeEvents();
-
-	UMoverComponent& GetMoverComponent() const;
-	Chaos::FPhysicsSolver* GetPhysicsSolver() const;
-	Chaos::FPhysicsObject* GetPhysicsObject() const;
-	Chaos::FPBDRigidParticle* GetControlledParticle() const;
-	FMoverTimeStep GetMoverTimeStep(double DeltaTimeInSeconds, int32 NetworkPhysicsTickOffset) const;
+	CHAOSMOVER_API virtual void InitSimulation();
+	CHAOSMOVER_API virtual void DeinitSimulation();
+	CHAOSMOVER_API virtual void CreatePhysics();
+	CHAOSMOVER_API virtual void DestroyPhysics();
+	CHAOSMOVER_API void CreateCharacterGroundConstraint();
+	CHAOSMOVER_API void DestroyCharacterGroundConstraint();
+	CHAOSMOVER_API void CreatePathTargetConstraint();
+	CHAOSMOVER_API void DestroyPathTargetConstraint();
+
+	CHAOSMOVER_API virtual void GenerateInput(const FMoverTimeStep& TimeStep, UE::ChaosMover::FSimulationInputData& InputData);
+
+	CHAOSMOVER_API UMoverComponent& GetMoverComponent() const;
+	CHAOSMOVER_API Chaos::FPhysicsSolver* GetPhysicsSolver() const;
+	CHAOSMOVER_API Chaos::FPhysicsObject* GetPhysicsObject() const;
+	CHAOSMOVER_API Chaos::FPBDRigidParticle* GetControlledParticle() const;
 
 	UFUNCTION()
-	virtual void HandleUpdatedComponentPhysicsStateChanged(UPrimitiveComponent* ChangedComponent, EComponentPhysicsStateChange StateChange);
+	CHAOSMOVER_API virtual void HandleUpdatedComponentPhysicsStateChanged(UPrimitiveComponent* ChangedComponent, EComponentPhysicsStateChange StateChange);
 
 	UFUNCTION()
-	void HandleOwningPawnControllerChanged_Server(APawn* OwnerPawn, AController* OldController, AController* NewController);
+	CHAOSMOVER_API void HandleOwningPawnControllerChanged_Server(APawn* OwnerPawn, AController* OldController, AController* NewController);
 
 	UPROPERTY(Transient)
 	TObjectPtr<UNetworkPhysicsComponent> NetworkPhysicsComponent;
diff --git a/Engine/Plugins/Experimental/ChaosMover/Source/ChaosMover/Public/ChaosMover/ChaosMovementModeTransition.h b/Engine/Plugins/Experimental/ChaosMover/Source/ChaosMover/Public/ChaosMover/ChaosMovementModeTransition.h
--- a/Engine/Plugins/Experimental/ChaosMover/Source/ChaosMover/Public/ChaosMover/ChaosMovementModeTransition.h
+++ b/Engine/Plugins/Experimental/ChaosMover/Source/ChaosMover/Public/ChaosMover/ChaosMovementModeTransition.h
@@ -11,13 +11,13 @@ class UChaosMoverSimulation;
 /**
  * Base class for all Chaos movement mode transitions
  */
-UCLASS(Abstract, Blueprintable, BlueprintType, EditInlineNew, DefaultToInstanced)
-class CHAOSMOVER_API UChaosMovementModeTransition : public UBaseMovementModeTransition
+UCLASS(MinimalAPI, Abstract, Blueprintable, BlueprintType, EditInlineNew, DefaultToInstanced)
+class UChaosMovementModeTransition : public UBaseMovementModeTransition
 {
 	GENERATED_BODY()
 
 public:
-	UChaosMovementModeTransition(const FObjectInitializer& ObjectInitializer);
+	CHAOSMOVER_API UChaosMovementModeTransition(const FObjectInitializer& ObjectInitializer);
 
 	UFUNCTION(BlueprintPure, Category = ChaosMover)
 	const UChaosMoverSimulation* GetSimulation() const
@@ -25,6 +25,12 @@ public:
 		return Simulation;
 	}
 
+	UFUNCTION(BlueprintPure, Category = ChaosMover)
+	UChaosMoverSimulation* GetSimulation_Mutable()
+	{
+		return Simulation;
+	}
+
 	void SetSimulation(UChaosMoverSimulation* InSimulation)
 	{
 		Simulation = InSimulation;
diff --git a/Engine/Plugins/Experimental/ChaosMover/Source/ChaosMover/Public/ChaosMover/ChaosMoverSimulation.h b/Engine/Plugins/Experimental/ChaosMover/Source/ChaosMover/Public/ChaosMover/ChaosMoverSimulation.h
--- a/Engine/Plugins/Experimental/ChaosMover/Source/ChaosMover/Public/ChaosMover/ChaosMoverSimulation.h
+++ b/Engine/Plugins/Experimental/ChaosMover/Source/ChaosMover/Public/ChaosMover/ChaosMoverSimulation.h
@@ -20,37 +20,35 @@ namespace Chaos
 	class FCollisionContactModifier;
 }
 
-UCLASS(BlueprintType)
-class CHAOSMOVER_API UChaosMoverSimulation : public UMoverSimulation
+UCLASS(MinimalAPI, BlueprintType)
+class UChaosMoverSimulation : public UMoverSimulation
 {
 	GENERATED_BODY()
 
 public:
-	UChaosMoverSimulation();
-
-	virtual void OnMovementModeChanged(FName PreviousMovementModeName, TWeakObjectPtr<UBaseMovementMode> PreviousMovementMode, FName NextMovementModeName, TWeakObjectPtr<UBaseMovementMode> NextMovementMode);
+	CHAOSMOVER_API UChaosMoverSimulation();
 
 	// Returns the local simulation input MoverDataCollection, to read local non networked data passed to the simulation by the gameplay thread
 	UFUNCTION(BlueprintPure, Category = Mover)
-	const FMoverDataCollection& GetLocalSimInput() const;
+	CHAOSMOVER_API const FMoverDataCollection& GetLocalSimInput() const;
 
 	// Returns the local simulation input MoverDataCollection, to pass local non networked data to the simulation
 	// Only available from the gameplay thread
 	UFUNCTION(BlueprintPure, Category = Mover)
-	FMoverDataCollection& GetLocalSimInput_Mutable();
+	CHAOSMOVER_API FMoverDataCollection& GetLocalSimInput_Mutable();
 
 	UFUNCTION(BlueprintPure, Category = Mover)
-	const UBaseMovementMode* GetCurrentMovementMode() const;
+	CHAOSMOVER_API const UBaseMovementMode* GetCurrentMovementMode() const;
 
 	UFUNCTION(BlueprintPure, Category = Mover)
-	const UBaseMovementMode* FindMovementModeByName(const FName& Name) const;
+	CHAOSMOVER_API const UBaseMovementMode* FindMovementModeByName(const FName& Name) const;
 
 	UFUNCTION(BlueprintCallable, CustomThunk, Category = Mover, meta = (CustomStructureParam = "EffectAsRawData", AllowAbstract = "false", DisplayName = "Queue Instant Movement Effect"))
 	void K2_QueueInstantMovementEffect(UPARAM(DisplayName = "Instant Movement Effect") const int32& EffectAsRawData);
 	DECLARE_FUNCTION(execK2_QueueInstantMovementEffect);
 
 	// Queue a Instant Movement Effect to take place at the end of this frame or start of the next subtick - whichever happens first
-	void QueueInstantMovementEffect(TSharedPtr<FInstantMovementEffect> Move);
+	CHAOSMOVER_API void QueueInstantMovementEffect(TSharedPtr<FInstantMovementEffect> Move);
 
 	struct FInitParams
 	{
@@ -69,61 +67,62 @@ public:
 		UWorld* World = nullptr;
 	};
 
-	void Init(const FInitParams& InitParams);
-	void Deinit();
-	void SimulationTick(const FMoverTimeStep& TimeStep, const UE::ChaosMover::FSimulationInputData& InputData, UE::ChaosMover::FSimulationOutputData& OutputData);
-	void ModifyContacts(const FMoverTimeStep& TimeStep, const UE::ChaosMover::FSimulationInputData& InputData, const UE::ChaosMover::FSimulationOutputData& OutputData, Chaos::FCollisionContactModifier& Modifier) const;
+	CHAOSMOVER_API void Init(const FInitParams& InitParams);
+	CHAOSMOVER_API void Deinit();
+	CHAOSMOVER_API void SimulationTick(const FMoverTimeStep& TimeStep, const UE::ChaosMover::FSimulationInputData& InputData, UE::ChaosMover::FSimulationOutputData& OutputData);
+	CHAOSMOVER_API void ModifyContacts(const FMoverTimeStep& TimeStep, const UE::ChaosMover::FSimulationInputData& InputData, const UE::ChaosMover::FSimulationOutputData& OutputData, Chaos::FCollisionContactModifier& Modifier) const;
 
-	void GetAndClearQueuedModeChangeEvents(TArray<UE::ChaosMover::FMovementModeChangeEvent>& OutModeChangedEvents);
+	CHAOSMOVER_API void AddEvent(TSharedPtr<FMoverSimulationEventData> Event);
 
-	void InitNetInputData(const FMoverInputCmdContext& InNetInputCmd);
-	void ApplyNetInputData(const FMoverInputCmdContext& InNetInputCmd);
-	void BuildNetInputData(FMoverInputCmdContext& OutNetInputCmd) const;
-	void ApplyNetStateData(const FMoverSyncState& InNetSyncState);
-	void BuildNetStateData(FMoverSyncState& OutNetSyncState) const;
+	CHAOSMOVER_API void InitNetInputData(const FMoverInputCmdContext& InNetInputCmd);
+	CHAOSMOVER_API void ApplyNetInputData(const FMoverInputCmdContext& InNetInputCmd);
+	CHAOSMOVER_API void BuildNetInputData(FMoverInputCmdContext& OutNetInputCmd) const;
+	CHAOSMOVER_API void ApplyNetStateData(const FMoverSyncState& InNetSyncState);
+	CHAOSMOVER_API void BuildNetStateData(FMoverSyncState& OutNetSyncState) const;
 
 	//~ Debugging Util functions
 	// Collection for holding extra debug data, that will be sent to the Chaos Visual Debugger for debugging
 	UFUNCTION(BlueprintPure, Category = Mover)
-	FMoverDataCollection& GetDebugSimData();
+	CHAOSMOVER_API FMoverDataCollection& GetDebugSimData();
 	//~ End of Debugging Util functions
 
 protected:
-	virtual void OnInit();
-	virtual void OnDeinit();
-	virtual void OnPreSimulationTick(const FMoverTimeStep& TimeStep, const UE::ChaosMover::FSimulationInputData& InputData);
-	virtual void OnSimulationTick(const FMoverTimeStep& TimeStep, const UE::ChaosMover::FSimulationInputData& InputData, UE::ChaosMover::FSimulationOutputData& OutputData);
-	virtual void OnPostSimulationTick(const FMoverTimeStep& TimeStep, UE::ChaosMover::FSimulationOutputData& OutputData);
-	virtual void OnModifyContacts(const FMoverTimeStep& TimeStep, const UE::ChaosMover::FSimulationInputData& InputData, const UE::ChaosMover::FSimulationOutputData& OutputData, Chaos::FCollisionContactModifier& Modifier) const;
+	CHAOSMOVER_API virtual void OnInit();
+	CHAOSMOVER_API virtual void OnDeinit();
+	CHAOSMOVER_API virtual void OnPreSimulationTick(const FMoverTimeStep& TimeStep, const UE::ChaosMover::FSimulationInputData& InputData);
+	CHAOSMOVER_API virtual void OnSimulationTick(const FMoverTimeStep& TimeStep, const UE::ChaosMover::FSimulationInputData& InputData, UE::ChaosMover::FSimulationOutputData& OutputData);
+	CHAOSMOVER_API virtual void OnPostSimulationTick(const FMoverTimeStep& TimeStep, UE::ChaosMover::FSimulationOutputData& OutputData);
+	CHAOSMOVER_API virtual void OnModifyContacts(const FMoverTimeStep& TimeStep, const UE::ChaosMover::FSimulationInputData& InputData, const UE::ChaosMover::FSimulationOutputData& OutputData, Chaos::FCollisionContactModifier& Modifier) const;
 
 	// Character-like movement sim steps
-	virtual void PreSimulationTickCharacter(const IChaosCharacterMovementModeInterface& CharacterMode, const FMoverTimeStep& TimeStep, const UE::ChaosMover::FSimulationInputData& InputData);
-	virtual void PostSimulationTickCharacter(const IChaosCharacterMovementModeInterface& CharacterMode, const FMoverTimeStep& TimeStep, UE::ChaosMover::FSimulationOutputData& OutputData);
-	virtual void PostSimulationTickCharacterConstraint(const IChaosCharacterConstraintMovementModeInterface& CharacterConstraintMode, const FMoverTimeStep& TimeStep, UE::ChaosMover::FSimulationOutputData& OutputData);
+	CHAOSMOVER_API virtual void PreSimulationTickCharacter(const IChaosCharacterMovementModeInterface& CharacterMode, const FMoverTimeStep& TimeStep, const UE::ChaosMover::FSimulationInputData& InputData);
+	CHAOSMOVER_API virtual void PostSimulationTickCharacter(const IChaosCharacterMovementModeInterface& CharacterMode, const FMoverTimeStep& TimeStep, UE::ChaosMover::FSimulationOutputData& OutputData);
+	CHAOSMOVER_API virtual void PostSimulationTickCharacterConstraint(const IChaosCharacterConstraintMovementModeInterface& CharacterConstraintMode, const FMoverTimeStep& TimeStep, UE::ChaosMover::FSimulationOutputData& OutputData);
 
 	// Pathed movement sim steps (kinematic or constrained)
-	virtual void PostSimulationTickPathedMovement(const IChaosPathedMovementModeInterface& ConstraintMode, const FMoverTimeStep& TimeStep, UE::ChaosMover::FSimulationOutputData& OutputData);
+	CHAOSMOVER_API virtual void PostSimulationTickPathedMovement(const IChaosPathedMovementModeInterface& ConstraintMode, const FMoverTimeStep& TimeStep, UE::ChaosMover::FSimulationOutputData& OutputData);
 
-	Chaos::FPBDRigidParticleHandle* GetControlledParticle() const;
+	CHAOSMOVER_API Chaos::FPBDRigidParticleHandle* GetControlledParticle() const;
 
-	void EnqueueModeChangeEvent(const UE::ChaosMover::FMovementModeChangeEvent& MovementModeChangeEvent);
+	CHAOSMOVER_API virtual void ProcessSimulationEvent(const FMoverSimulationEventData& ModeChangedData);
+	CHAOSMOVER_API virtual void OnMovementModeChanged(const FMovementModeChangedEventData& ModeChangedData);
 
 	// Object State.
 	// In kinematic mode, the updated primitive component needs both its SimulatePhysics and UpdateKinematicFromSimulation options checked
-	void SetControlledParticleDynamic();
-	void SetControlledParticleKinematic();
-	bool IsControlledParticleDynamic() const;
-	bool IsControlledParticleKinematic() const;
+	CHAOSMOVER_API void SetControlledParticleDynamic();
+	CHAOSMOVER_API void SetControlledParticleKinematic();
+	CHAOSMOVER_API bool IsControlledParticleDynamic() const;
+	CHAOSMOVER_API bool IsControlledParticleKinematic() const;
 
 	// Character ground constraint
-	void EnableCharacterConstraint();
-	void DisableCharacterConstraint();
+	CHAOSMOVER_API void EnableCharacterConstraint();
+	CHAOSMOVER_API void DisableCharacterConstraint();
 
 	// Path target constraint
-	void EnablePathTargetConstraint();
-	void DisablePathTargetConstraint();
+	CHAOSMOVER_API void EnablePathTargetConstraint();
+	CHAOSMOVER_API void DisablePathTargetConstraint();
 
-	void TraceMoverData(const UE::ChaosMover::FSimulationOutputData& OutputData);
+	CHAOSMOVER_API void TraceMoverData(const UE::ChaosMover::FSimulationOutputData& OutputData);
 
 	// State structs
 	FMoverSyncState CurrentSyncState;
@@ -155,14 +154,7 @@ protected:
 
 private:
 	FMoverInputCmdContext InputCmd;
-
-	// This queue is used to record movement mode transitions and postpone the execution of their callbacks
-	// This implementation is flawed, even with a mutex, we need to implement differently
-	TArray<UE::ChaosMover::FMovementModeChangeEvent> QueuedModeChangeEvents;
-	// FinalizeFrame and the next physics tick can be concurrent so QueuedModeChangeEvents access was not thread safe
-	// This mutex lock is just temporary to prevent crashes
-	// TODO: Implement a lockless solution (output the mode changed events as part of sim output instead)
-	mutable FRWLock ModeChangedEventsLock;
+	TArray<TSharedPtr<FMoverSimulationEventData>> Events;
 
 	bool bInputCmdOverridden = false;
 };
diff --git a/Engine/Plugins/Experimental/ChaosMover/Source/ChaosMover/Public/ChaosMover/Character/ChaosCharacterMoverComponent.h b/Engine/Plugins/Experimental/ChaosMover/Source/ChaosMover/Public/ChaosMover/Character/ChaosCharacterMoverComponent.h
new file mode 100644
--- /dev/null
+++ b/Engine/Plugins/Experimental/ChaosMover/Source/ChaosMover/Public/ChaosMover/Character/ChaosCharacterMoverComponent.h
@@ -0,0 +1,34 @@
+// Copyright Epic Games, Inc. All Rights Reserved.
+
+#pragma once
+
+#include "DefaultMovementSet/CharacterMoverComponent.h"
+#include "MoveLibrary/FloorQueryUtils.h"
+
+#include "ChaosCharacterMoverComponent.generated.h"
+
+// Fired after the actor lands on a valid surface. First param is the name of the mode this actor is in after landing. Second param is the hit result from hitting the floor.
+DECLARE_DYNAMIC_MULTICAST_DELEGATE_TwoParams(FChaosMover_OnLanded, const FName&, NextMovementModeName, const FHitResult&, HitResult);
+
+
+UCLASS(MinimalAPI, BlueprintType, Blueprintable, meta = (BlueprintSpawnableComponent))
+class UChaosCharacterMoverComponent : public UCharacterMoverComponent
+{
+	GENERATED_BODY()
+
+public:
+	CHAOSMOVER_API UChaosCharacterMoverComponent();
+
+	CHAOSMOVER_API virtual bool TryGetFloorCheckHitResult(FHitResult& OutHitResult) const override;
+
+	// Broadcast when this actor lands on a valid surface.
+	UPROPERTY(BlueprintAssignable, Category = Mover)
+	FChaosMover_OnLanded OnLanded;
+
+protected:
+	CHAOSMOVER_API virtual void ProcessSimulationEvent(const FMoverSimulationEventData& EventData) override;
+	CHAOSMOVER_API virtual void SetAdditionalSimulationOutput(const FMoverDataCollection& Data) override;
+
+	bool bFloorResultSet = false;
+	FFloorCheckResult LatestFloorResult;
+};
diff --git a/Engine/Plugins/Experimental/Mover/Source/Mover/Private/DefaultMovementSet/CharacterMoverSimulationTypes.cpp b/Engine/Plugins/Experimental/Mover/Source/Mover/Private/DefaultMovementSet/CharacterMoverSimulationTypes.cpp
new file mode 100644
--- /dev/null
+++ b/Engine/Plugins/Experimental/Mover/Source/Mover/Private/DefaultMovementSet/CharacterMoverSimulationTypes.cpp
@@ -0,0 +1,63 @@
+// Copyright Epic Games, Inc. All Rights Reserved.
+
+#include "DefaultMovementSet/CharacterMoverSimulationTypes.h"
+
+bool FFloorResultData::NetSerialize(FArchive& Ar, UPackageMap* Map, bool& bOutSuccess)
+{
+	Super::NetSerialize(Ar, Map, bOutSuccess);
+
+	bool bBlockingHit = (bool)FloorResult.bBlockingHit;
+	bool bLineTrace = (bool)FloorResult.bLineTrace;
+	bool bWalkableFloor = (bool)FloorResult.bWalkableFloor;
+	Ar.SerializeBits(&bBlockingHit, 1);
+	Ar.SerializeBits(&bLineTrace, 1);
+	Ar.SerializeBits(&bWalkableFloor, 1);
+
+	if (Ar.IsLoading())
+	{
+		FloorResult.bBlockingHit = bBlockingHit;
+		FloorResult.bLineTrace = bLineTrace;
+		FloorResult.bWalkableFloor = bWalkableFloor;
+	}
+
+	Ar << FloorResult.FloorDist;
+	FloorResult.HitResult.NetSerialize(Ar, Map, bOutSuccess);
+
+	return true;
+}
+
+void FFloorResultData::ToString(FAnsiStringBuilderBase& Out) const
+{
+	Super::ToString(Out);
+
+	Out.Appendf("bBlockingHit: %i | ", FloorResult.bBlockingHit);
+	Out.Appendf("bLineTrace: %i | ", FloorResult.bLineTrace);
+	Out.Appendf("bWalkableFloor: %i | ", FloorResult.bWalkableFloor);
+	Out.Appendf("FloorDist: %.2f/n", FloorResult.FloorDist);
+	Out.Appendf("HitResult: %s/n", *FloorResult.HitResult.ToString());
+}
+
+bool FFloorResultData::ShouldReconcile(const FMoverDataStructBase& AuthorityState) const
+{
+	return false;
+}
+
+void FFloorResultData::Interpolate(const FMoverDataStructBase& From, const FMoverDataStructBase& To, float Pct)
+{
+	if (Pct < 0.5f)
+	{
+		*this = static_cast<const FFloorResultData&>(From);
+	}
+	else
+	{
+		*this = static_cast<const FFloorResultData&>(To);
+	}
+}
+
+void FFloorResultData::Merge(const FMoverDataStructBase& From)
+{
+}
+
+void FFloorResultData::Decay(float DecayAmount)
+{
+}
\ No newline at end of file
diff --git a/Engine/Plugins/Experimental/Mover/Source/Mover/Private/MoverComponent.cpp b/Engine/Plugins/Experimental/Mover/Source/Mover/Private/MoverComponent.cpp
--- a/Engine/Plugins/Experimental/Mover/Source/Mover/Private/MoverComponent.cpp
+++ b/Engine/Plugins/Experimental/Mover/Source/Mover/Private/MoverComponent.cpp
@@ -2211,9 +2211,29 @@ void UMoverComponent::SetSimulationOutput(const FMoverTimeStep& TimeStep, const
 	LastMoverDefaultSyncState = BufferedSyncState.SyncStateCollection.FindDataByType<FMoverDefaultSyncState>();
 	MoverSyncStateDoubleBuffer.Flip();
 
+	for (const TSharedPtr<FMoverSimulationEventData>& EventData : OutputData.Events)
+	{
+		if (const FMoverSimulationEventData* Data = EventData.Get())
+		{
+			ProcessSimulationEvent(*Data);
+		}
+	}
 	// This is for things like the ground info that we want to cache and interpolate but isn't part of the networked sync state
 	// AdditionalOutputData is generic because ground info might not be useful for platforms, say, but platforms may want to store something else.
-	// OutputData.AdditionalOutputData;
+	SetAdditionalSimulationOutput(OutputData.AdditionalOutputData);
+}
+
+void UMoverComponent::ProcessSimulationEvent(const FMoverSimulationEventData& EventData)
+{
+	if (const FMovementModeChangedEventData* ModeChangedData = EventData.CastTo<FMovementModeChangedEventData>())
+	{
+		OnMovementModeChanged.Broadcast(ModeChangedData->PreviousModeName, ModeChangedData->NewModeName);
+	}
+}
+
+void UMoverComponent::SetAdditionalSimulationOutput(const FMoverDataCollection& Data)
+{
+
 }
 
 #undef LOCTEXT_NAMESPACE
diff --git a/Engine/Plugins/Experimental/Mover/Source/Mover/Private/MoverSimulationTypes.cpp b/Engine/Plugins/Experimental/Mover/Source/Mover/Private/MoverSimulationTypes.cpp
new file mode 100644
--- /dev/null
+++ b/Engine/Plugins/Experimental/Mover/Source/Mover/Private/MoverSimulationTypes.cpp
@@ -0,0 +1,99 @@
+// Copyright Epic Games, Inc. All Rights Reserved.
+
+#include "MoverSimulationTypes.h"
+
+UScriptStruct* FMoverSimulationEventData::GetScriptStruct() const
+{
+	checkf(false, TEXT("%hs is being called erroneously. This must be overridden in derived types!"), __FUNCTION__);
+	return FMoverSimulationEventData::StaticStruct();
+}
+
+namespace UE::Mover
+{
+	void FSimulationOutputData::Reset()
+	{
+		SyncState.Reset();
+		LastUsedInputCmd.Reset();
+		AdditionalOutputData.Empty();
+		Events.Empty();
+	}
+
+	void FSimulationOutputData::Interpolate(const FSimulationOutputData& From, const FSimulationOutputData& To, float Alpha, float SimTimeMs)
+	{
+		SyncState.Interpolate(&From.SyncState, &To.SyncState, Alpha);
+		LastUsedInputCmd.Interpolate(&From.LastUsedInputCmd, &To.LastUsedInputCmd, Alpha);
+		AdditionalOutputData.Interpolate(From.AdditionalOutputData, To.AdditionalOutputData, Alpha);
+
+		for (const TSharedPtr<FMoverSimulationEventData>& EventData : From.Events)
+		{
+			if (const FMoverSimulationEventData* DataPtr = EventData.Get())
+			{
+				if (DataPtr->EventTimeMs <= SimTimeMs)
+				{
+					Events.Add(EventData);
+				}
+			}
+		}
+		for (const TSharedPtr<FMoverSimulationEventData>& EventData : To.Events)
+		{
+			if (const FMoverSimulationEventData* DataPtr = EventData.Get())
+			{
+				if (DataPtr->EventTimeMs <= SimTimeMs)
+				{
+					Events.Add(EventData);
+				}
+			}
+		}
+	}
+
+	void FSimulationOutputRecord::FData::Reset()
+	{
+		TimeStep = FMoverTimeStep();
+		SimOutputData.Reset();
+	}
+
+	void FSimulationOutputRecord::Add(const FMoverTimeStep& InTimeStep, const FSimulationOutputData& InData)
+	{
+		CurrentIndex = (CurrentIndex + 1) % 2;
+		Data[CurrentIndex] = { InTimeStep, InData };
+	}
+
+	const FSimulationOutputData& FSimulationOutputRecord::GetLatest() const
+	{
+		return Data[CurrentIndex].SimOutputData;
+	}
+
+	void FSimulationOutputRecord::GetInterpolated(float AtBaseTimeMs, FMoverTimeStep& OutTimeStep, FSimulationOutputData& OutData) const
+	{
+		const uint8 PrevIndex = (CurrentIndex + 1) % 2;
+		const float PrevTimeMs = Data[PrevIndex].TimeStep.BaseSimTimeMs;
+		const float CurrTimeMs = Data[CurrentIndex].TimeStep.BaseSimTimeMs;
+
+		if (FMath::IsNearlyEqual(PrevTimeMs, CurrTimeMs) || (AtBaseTimeMs >= CurrTimeMs))
+		{
+			OutData = Data[CurrentIndex].SimOutputData;
+			OutTimeStep = Data[CurrentIndex].TimeStep;
+		}
+		else if (AtBaseTimeMs <= PrevTimeMs)
+		{
+			OutData = Data[PrevIndex].SimOutputData;
+			OutTimeStep = Data[PrevIndex].TimeStep;
+		}
+		else
+		{
+			const float Alpha = FMath::Clamp((AtBaseTimeMs - PrevTimeMs) / (CurrTimeMs - PrevTimeMs), 0.0f, 1.0f);
+			OutData.Interpolate(Data[PrevIndex].SimOutputData, Data[CurrentIndex].SimOutputData, Alpha, AtBaseTimeMs);
+			OutTimeStep = Data[PrevIndex].TimeStep;
+		}
+
+		OutTimeStep.BaseSimTimeMs = AtBaseTimeMs;
+	}
+
+	void FSimulationOutputRecord::Clear()
+	{
+		CurrentIndex = 1;
+		Data[0].Reset();
+		Data[1].Reset();
+	}
+
+} // namespace UE::Mover
\ No newline at end of file
diff --git a/Engine/Plugins/Experimental/Mover/Source/Mover/Public/DefaultMovementSet/CharacterMoverSimulationTypes.h b/Engine/Plugins/Experimental/Mover/Source/Mover/Public/DefaultMovementSet/CharacterMoverSimulationTypes.h
new file mode 100644
--- /dev/null
+++ b/Engine/Plugins/Experimental/Mover/Source/Mover/Public/DefaultMovementSet/CharacterMoverSimulationTypes.h
@@ -0,0 +1,69 @@
+// Copyright Epic Games, Inc. All Rights Reserved.
+
+#pragma once
+
+#include "MoverSimulationTypes.h"
+#include "MoveLibrary/FloorQueryUtils.h"
+
+#include "CharacterMoverSimulationTypes.generated.h"
+
+USTRUCT()
+struct FLandedEventData : public FMoverSimulationEventData
+{
+	GENERATED_BODY()
+
+	FLandedEventData(float InEventTimeMs, const FHitResult& InHitResult, const FName InNewModeName)
+		: FMoverSimulationEventData(InEventTimeMs)
+		, HitResult(InHitResult)
+		, NewModeName(InNewModeName)
+	{
+	}
+	FLandedEventData() {}
+
+	virtual UScriptStruct* GetScriptStruct() const override
+	{
+		return FLandedEventData::StaticStruct();
+	}
+
+	FHitResult HitResult;
+	FName NewModeName = NAME_None;
+};
+
+USTRUCT(BlueprintType)
+struct FFloorResultData : public FMoverDataStructBase
+{
+	GENERATED_BODY()
+
+public:
+	UPROPERTY(BlueprintReadWrite, Category = Mover)
+	FFloorCheckResult FloorResult;
+
+	FFloorResultData() = default;
+	virtual ~FFloorResultData() = default;
+
+	virtual UScriptStruct* GetScriptStruct() const override
+	{
+		return StaticStruct();
+	}
+
+	virtual FMoverDataStructBase* Clone() const override
+	{
+		return new FFloorResultData(*this);
+	}
+
+	MOVER_API virtual bool NetSerialize(FArchive& Ar, UPackageMap* Map, bool& bOutSuccess) override;
+	MOVER_API virtual void ToString(FAnsiStringBuilderBase& Out) const override;
+	MOVER_API virtual bool ShouldReconcile(const FMoverDataStructBase& AuthorityState) const override;
+	MOVER_API virtual void Interpolate(const FMoverDataStructBase& From, const FMoverDataStructBase& To, float Pct) override;
+	MOVER_API virtual void Merge(const FMoverDataStructBase& From) override;
+	MOVER_API virtual void Decay(float DecayAmount) override;
+};
+
+template<>
+struct TStructOpsTypeTraits< FFloorResultData > : public TStructOpsTypeTraitsBase2< FFloorResultData >
+{
+	enum
+	{
+		WithCopy = true
+	};
+};
\ No newline at end of file
diff --git a/Engine/Plugins/Experimental/Mover/Source/Mover/Public/MoverComponent.h b/Engine/Plugins/Experimental/Mover/Source/Mover/Public/MoverComponent.h
--- a/Engine/Plugins/Experimental/Mover/Source/Mover/Public/MoverComponent.h
+++ b/Engine/Plugins/Experimental/Mover/Source/Mover/Public/MoverComponent.h
@@ -377,7 +377,7 @@ public:	// Queries
 
 	// Access the most recent floor check hit result.
 	UFUNCTION(BlueprintPure, Category = Mover)
-	MOVER_API bool TryGetFloorCheckHitResult(FHitResult& OutHitResult) const;
+	MOVER_API virtual bool TryGetFloorCheckHitResult(FHitResult& OutHitResult) const;
 
 	// Access the read-only version of the Mover's Blackboard
 	UFUNCTION(BlueprintPure, Category=Mover)
@@ -636,6 +636,10 @@ public:
 
 	MOVER_API void SetSimulationOutput(const FMoverTimeStep& TimeStep, const UE::Mover::FSimulationOutputData& OutputData);
 
+protected:
+	MOVER_API virtual void ProcessSimulationEvent(const FMoverSimulationEventData& EventData);
+	MOVER_API virtual void SetAdditionalSimulationOutput(const FMoverDataCollection& Data);
+
 private:
 	/** Whether or not gravity is overridden on this actor. Otherwise, fall back on world settings. See @SetGravityOverride */
 	UPROPERTY(EditDefaultsOnly, Category="Mover|Gravity")
diff --git a/Engine/Plugins/Experimental/Mover/Source/Mover/Public/MoverSimulationTypes.h b/Engine/Plugins/Experimental/Mover/Source/Mover/Public/MoverSimulationTypes.h
--- a/Engine/Plugins/Experimental/Mover/Source/Mover/Public/MoverSimulationTypes.h
+++ b/Engine/Plugins/Experimental/Mover/Source/Mover/Public/MoverSimulationTypes.h
@@ -391,27 +391,68 @@ struct FMoverPredictTrajectoryParams
 
 };
 
+USTRUCT()
+struct FMoverSimulationEventData
+{
+	GENERATED_BODY()
+
+	FMoverSimulationEventData(float InEventTimeMs)
+		: EventTimeMs(InEventTimeMs)
+	{
+	}
+	FMoverSimulationEventData() {}
+
+	// User must override
+	MOVER_API virtual UScriptStruct* GetScriptStruct() const;
+
+	template<typename T>
+	T* CastTo_Mutable()
+	{
+		return T::StaticStruct() == GetScriptStruct() ? static_cast<T*>(this) : nullptr;
+	}
+
+	template<typename T>
+	const T* CastTo() const
+	{
+		return const_cast<const T*>(const_cast<FMoverSimulationEventData*>(this)->CastTo_Mutable<T>());
+	}
+
+	float EventTimeMs = 0.0f;
+};
+
+USTRUCT()
+struct FMovementModeChangedEventData : public FMoverSimulationEventData
+{
+	GENERATED_BODY()
+
+	FMovementModeChangedEventData(float InEventTimeMs, const FName InPreviousModeName, const FName InNewModeName)
+		: FMoverSimulationEventData(InEventTimeMs)
+		, PreviousModeName(InPreviousModeName)
+		, NewModeName(InNewModeName)
+	{
+	}
+	FMovementModeChangedEventData() {}
+
+	virtual UScriptStruct* GetScriptStruct() const override
+	{
+		return FMovementModeChangedEventData::StaticStruct();
+	}
+
+	FName PreviousModeName = NAME_None;
+	FName NewModeName = NAME_None;
+};
+
 namespace UE::Mover
 {
 	struct FSimulationOutputData
 	{
-		void Reset()
-		{
-			SyncState.Reset();
-			LastUsedInputCmd.Reset();
-			AdditionalOutputData.Empty();
-		}
-
-		void Interpolate(const FSimulationOutputData& From, const FSimulationOutputData& To, float Alpha)
-		{
-			SyncState.Interpolate(&From.SyncState, &To.SyncState, Alpha);
-			LastUsedInputCmd.Interpolate(&From.LastUsedInputCmd, &To.LastUsedInputCmd, Alpha);
-			AdditionalOutputData.Interpolate(From.AdditionalOutputData, To.AdditionalOutputData, Alpha);
-		}
+		MOVER_API void Reset();
+		MOVER_API void Interpolate(const FSimulationOutputData& From, const FSimulationOutputData& To, float Alpha, float SimTimeMs);
 
 		FMoverSyncState SyncState;
 		FMoverInputCmdContext LastUsedInputCmd;
 		FMoverDataCollection AdditionalOutputData;
+		TArray<TSharedPtr<FMoverSimulationEventData>> Events;
 	};
 
 	class FSimulationOutputRecord
@@ -419,65 +460,18 @@ namespace UE::Mover
 	public:
 		struct FData
 		{
-			void Reset()
-			{
-				TimeStep = FMoverTimeStep();
-				SimOutputData.Reset();
-			}
+			MOVER_API void Reset();
 
 			FMoverTimeStep TimeStep;
 			FSimulationOutputData SimOutputData;
 		};
-		void Add(const FMoverTimeStep& InTimeStep , const FSimulationOutputData& InData)
-		{
-			CurrentIndex = (CurrentIndex + 1) % 2;
-			Data[CurrentIndex] = { InTimeStep, InData };
-		}
 
-		const FSimulationOutputData& GetLatest() const
-		{
-			return Data[CurrentIndex].SimOutputData;
-		}
-
-		void GetInterpolated(float SimTime, FMoverTimeStep& OutTimeStep, FSimulationOutputData& OutData) const
-		{
-			const uint8 PrevIndex = (CurrentIndex + 1) % 2;
-			const float PrevTime = Data[PrevIndex].TimeStep.BaseSimTimeMs * 0.001f;
-			const float CurrTime = Data[CurrentIndex].TimeStep.BaseSimTimeMs * 0.001f;
-
-			if (FMath::IsNearlyEqual(PrevTime, CurrTime) || (SimTime >= CurrTime))
-			{
-				OutData = Data[CurrentIndex].SimOutputData;
-				OutTimeStep = Data[CurrentIndex].TimeStep;
-			}
-			else if (SimTime <= PrevTime)
-			{
-				OutData = Data[PrevIndex].SimOutputData;
-				OutTimeStep = Data[PrevIndex].TimeStep;
-			}
-			else
-			{
-				const float Alpha = FMath::Clamp((SimTime - PrevTime) / (CurrTime - PrevTime), 0.0f, 1.0f);
-				OutData.Interpolate(Data[PrevIndex].SimOutputData, Data[CurrentIndex].SimOutputData, Alpha);
-				if (Alpha < 0.5f)
-				{
-					OutTimeStep = Data[PrevIndex].TimeStep;
-				}
-				else
-				{
-					OutTimeStep = Data[CurrentIndex].TimeStep;
-				}
-				OutTimeStep.BaseSimTimeMs = SimTime;
-			}
-		}
-
-		void Clear()
-		{
-			CurrentIndex = 1;
-			Data[0].Reset();
-			Data[1].Reset();
-		}
+		MOVER_API void Add(const FMoverTimeStep& InTimeStep, const FSimulationOutputData& InData);
+		MOVER_API const FSimulationOutputData& GetLatest() const;
+		MOVER_API void GetInterpolated(float AtBaseTimeMs, FMoverTimeStep& OutTimeStep, FSimulationOutputData& OutData) const;
+		MOVER_API void Clear();
 
+	private:
 		FData Data[2];
 		uint8 CurrentIndex = 1;
 	};
