Chaos Mover - Add some fixes to water check transition and remove the default add of the transition to walking and falling modes

#rnx

diff --git a/Engine/Plugins/Experimental/ChaosMover/Source/ChaosMover/Private/Character/Modes/ChaosFallingMode.cpp b/Engine/Plugins/Experimental/ChaosMover/Source/ChaosMover/Private/Character/Modes/ChaosFallingMode.cpp
--- a/Engine/Plugins/Experimental/ChaosMover/Source/ChaosMover/Private/Character/Modes/ChaosFallingMode.cpp
+++ b/Engine/Plugins/Experimental/ChaosMover/Source/ChaosMover/Private/Character/Modes/ChaosFallingMode.cpp
@@ -6,7 +6,6 @@
 #include "ChaosMover/ChaosMoverSimulation.h"
 #include "ChaosMover/Character/Settings/SharedChaosCharacterMovementSettings.h"
 #include "ChaosMover/Character/Transitions/ChaosCharacterLandingCheck.h"
-#include "ChaosMover/Character/Transitions/ChaosCharacterWaterCheck.h"
 #include "ChaosMover/Utilities/ChaosMoverQueryUtils.h"
 #include "MoveLibrary/AirMovementUtils.h"
 #include "MoveLibrary/FloorQueryUtils.h"
@@ -37,7 +36,6 @@ UChaosFallingMode::UChaosFallingMode(const FObjectInitializer& ObjectInitializer
 	TwistTorqueLimit = 0.0f;
 
 	Transitions.Add(CreateDefaultSubobject<UChaosCharacterLandingCheck>(TEXT("DefaultLandingCheck")));
-	Transitions.Add(CreateDefaultSubobject<UChaosCharacterWaterCheck>(TEXT("DefaultWaterCheck")));
 }
 
 void UChaosFallingMode::GenerateMove_Implementation(const FMoverTickStartData& StartState, const FMoverTimeStep& TimeStep, FProposedMove& OutProposedMove) const
diff --git a/Engine/Plugins/Experimental/ChaosMover/Source/ChaosMover/Private/Character/Modes/ChaosWalkingMode.cpp b/Engine/Plugins/Experimental/ChaosMover/Source/ChaosMover/Private/Character/Modes/ChaosWalkingMode.cpp
--- a/Engine/Plugins/Experimental/ChaosMover/Source/ChaosMover/Private/Character/Modes/ChaosWalkingMode.cpp
+++ b/Engine/Plugins/Experimental/ChaosMover/Source/ChaosMover/Private/Character/Modes/ChaosWalkingMode.cpp
@@ -11,7 +11,6 @@
 #include "ChaosMover/Character/Settings/SharedChaosCharacterMovementSettings.h"
 #include "ChaosMover/Character/Transitions/ChaosCharacterFallingCheck.h"
 #include "ChaosMover/Character/Transitions/ChaosCharacterJumpCheck.h"
-#include "ChaosMover/Character/Transitions/ChaosCharacterWaterCheck.h"
 #include "ChaosMover/Utilities/ChaosGroundMovementUtils.h"
 #include "ChaosMover/Utilities/ChaosMoverQueryUtils.h"
 #include "Math/UnitConversion.h"
@@ -37,7 +36,6 @@ UChaosWalkingMode::UChaosWalkingMode(const FObjectInitializer& ObjectInitializer
 
 	Transitions.Add(CreateDefaultSubobject<UChaosCharacterFallingCheck>(TEXT("DefaultFallingCheck")));
 	Transitions.Add(CreateDefaultSubobject<UChaosCharacterJumpCheck>(TEXT("DefaultJumpCheck")));
-	Transitions.Add(CreateDefaultSubobject<UChaosCharacterWaterCheck>(TEXT("DefaultWaterCheck")));
 }
 
 void UChaosWalkingMode::UpdateConstraintSettings(Chaos::FCharacterGroundConstraintSettings& ConstraintSettings) const
diff --git a/Engine/Plugins/Experimental/ChaosMover/Source/ChaosMover/Private/Character/Transitions/ChaosCharacterWaterCheck.cpp b/Engine/Plugins/Experimental/ChaosMover/Source/ChaosMover/Private/Character/Transitions/ChaosCharacterWaterCheck.cpp
--- a/Engine/Plugins/Experimental/ChaosMover/Source/ChaosMover/Private/Character/Transitions/ChaosCharacterWaterCheck.cpp
+++ b/Engine/Plugins/Experimental/ChaosMover/Source/ChaosMover/Private/Character/Transitions/ChaosCharacterWaterCheck.cpp
@@ -15,7 +15,7 @@ UChaosCharacterWaterCheck::UChaosCharacterWaterCheck(const FObjectInitializer& O
 	: Super(ObjectInitializer)
 {
 	bSupportsAsync = true;
-
+ 
 	WaterModeName = DefaultModeNames::Swimming;
 	GroundModeName = DefaultModeNames::Walking;
 	AirModeName = DefaultModeNames::Falling;
@@ -33,7 +33,7 @@ FTransitionEvalResult UChaosCharacterWaterCheck::Evaluate_Implementation(const F
 	{
 		const UMoverBlackboard* Blackboard = Simulation->GetBlackboard();
 		check(Blackboard);
-
+ 
 		const FCharacterDefaultInputs* CharacterInputs = StartState.InputCmd.InputCollection.FindDataByType<FCharacterDefaultInputs>();
 		const FMoverDefaultSyncState* StartingSyncState = StartState.SyncState.SyncStateCollection.FindDataByType<FMoverDefaultSyncState>();
 		if (!CharacterInputs || !StartingSyncState)
@@ -47,7 +47,7 @@ FTransitionEvalResult UChaosCharacterWaterCheck::Evaluate_Implementation(const F
 			const FVector LinearVelocity = StartingSyncState->GetVelocity_WorldSpace();
 			const float VerticalVelocity = SimInputs->UpDir.Dot(LinearVelocity);
 			const bool bJumping = CharacterInputs && CharacterInputs->bIsJumpJustPressed;
-			const bool bIsMovingUp = bJumping || (VerticalVelocity > 0.0f);
+			const bool bIsMovingUp = bJumping || (VerticalVelocity > UE_KINDA_SMALL_NUMBER);
 
 			FFindFloorResult FloorResult;
 			FWaterCheckResult WaterResult;
@@ -57,7 +57,8 @@ FTransitionEvalResult UChaosCharacterWaterCheck::Evaluate_Implementation(const F
 			
 			const float ProjectedImmersionDepth = WaterResult.WaterSplineData.ImmersionDepth - VerticalVelocity * DeltaSeconds;
 			const bool bStartSwimming = ProjectedImmersionDepth > WaterModeStartImmersionDepth;
-			const bool bStopSwimming = ProjectedImmersionDepth < WaterModeStopImmersionDepth;
+			const bool bInWater = StartState.SyncState.MovementMode == WaterModeName;
+			const bool bStopSwimming = bInWater && (ProjectedImmersionDepth < WaterModeStopImmersionDepth);
 
 			// Check if we need to enter/exit water mode
 			if (WaterResult.IsSwimmableVolume() && bStartSwimming)
@@ -78,8 +79,8 @@ FTransitionEvalResult UChaosCharacterWaterCheck::Evaluate_Implementation(const F
 			else if (bStopSwimming)
 			{
 				const bool bIsWithinReach = FloorResult.FloorDist <= Mode->GetTargetHeight();
-
-				if (FloorResult.IsWalkableFloor() && bIsWithinReach)
+ 
+				if (FloorResult.IsWalkableFloor() && bIsWithinReach && !bIsMovingUp)
 				{
 					if (!GroundModeName.IsNone())
 					{
@@ -118,4 +119,4 @@ FTransitionEvalResult UChaosCharacterWaterCheck::Evaluate_Implementation(const F
 
 void UChaosCharacterWaterCheck::Trigger_Implementation(const FSimulationTickParams& Params)
 {
-}
+}
\ No newline at end of file
