Chaos Mover - Rename aux data in simulation input

#rnx

diff --git a/Engine/Plugins/Experimental/ChaosMover/Source/ChaosMover/Private/ChaosMoverSimulation.cpp b/Engine/Plugins/Experimental/ChaosMover/Source/ChaosMover/Private/ChaosMoverSimulation.cpp
--- a/Engine/Plugins/Experimental/ChaosMover/Source/ChaosMover/Private/ChaosMoverSimulation.cpp
+++ b/Engine/Plugins/Experimental/ChaosMover/Source/ChaosMover/Private/ChaosMoverSimulation.cpp
@@ -289,8 +289,8 @@ void UChaosMoverSimulation::OnSimulationTick(const FMoverTimeStep& TimeStep, con
 {
 	check(Blackboard.Get());
 
-	FMoverTickStartData TickStartData(InputData.InputCmd, CurrentSyncState, InputData.LocalInputParams);
-	FMoverTickEndData TickEndData(&CurrentSyncState, &InputData.LocalInputParams);
+	FMoverTickStartData TickStartData(InputData.InputCmd, CurrentSyncState, InputData.AuxInputState);
+	FMoverTickEndData TickEndData(&CurrentSyncState, &InputData.AuxInputState);
 
 	StateMachine.OnSimulationTick(TimeStep, TickStartData, Blackboard.Get(), MovementMixerWeakPtr.Get(), TickEndData);
 
diff --git a/Engine/Plugins/Experimental/ChaosMover/Source/ChaosMover/Public/ChaosMover/ChaosMoverSimulationTypes.h b/Engine/Plugins/Experimental/ChaosMover/Source/ChaosMover/Public/ChaosMover/ChaosMoverSimulationTypes.h
--- a/Engine/Plugins/Experimental/ChaosMover/Source/ChaosMover/Public/ChaosMover/ChaosMoverSimulationTypes.h
+++ b/Engine/Plugins/Experimental/ChaosMover/Source/ChaosMover/Public/ChaosMover/ChaosMoverSimulationTypes.h
@@ -25,11 +25,11 @@ namespace UE::ChaosMover
 		void Reset()
 		{
 			InputCmd.Reset();
-			LocalInputParams.AuxStateCollection.Empty();
+			AuxInputState.AuxStateCollection.Empty();
 		}
 
-		mutable FMoverInputCmdContext InputCmd; // can be overridden by network physics
-		FMoverAuxStateContext LocalInputParams;
+		mutable FMoverInputCmdContext InputCmd; // Can be overridden by network physics
+		FMoverAuxStateContext AuxInputState; // Optional aux input state. Not networked
 	};
 
 	using FSimulationOutputData = ::UE::Mover::FSimulationOutputData;
