Chaos Mover - Add launch transition

diff --git a/Engine/Plugins/Experimental/ChaosMover/Source/ChaosMover/Private/Character/ChaosCharacterInputs.cpp b/Engine/Plugins/Experimental/ChaosMover/Source/ChaosMover/Private/Character/ChaosCharacterInputs.cpp
new file mode 100644
--- /dev/null
+++ b/Engine/Plugins/Experimental/ChaosMover/Source/ChaosMover/Private/Character/ChaosCharacterInputs.cpp
@@ -0,0 +1,45 @@
+// Copyright Epic Games, Inc. All Rights Reserved.
+
+#include "ChaosMover/Character/ChaosCharacterInputs.h"
+
+bool FChaosMoverLaunchInputs::NetSerialize(FArchive& Ar, UPackageMap* Map, bool& bOutSuccess)
+{
+	Super::NetSerialize(Ar, Map, bOutSuccess);
+
+	Ar << LaunchVelocityOrImpulse;
+	Ar << Mode;
+
+	bOutSuccess = true;
+	return true;
+}
+
+void FChaosMoverLaunchInputs::ToString(FAnsiStringBuilderBase& Out) const
+{
+	Super::ToString(Out);
+
+	Out.Appendf("LaunchVelocityOrImpulse: X=%.2f Y=%.2f Z=%.2f |", LaunchVelocityOrImpulse.X, LaunchVelocityOrImpulse.Y, LaunchVelocityOrImpulse.Z);
+	Out.Appendf("Mode: %u/n", Mode);
+}
+
+bool FChaosMoverLaunchInputs::ShouldReconcile(const FMoverDataStructBase& AuthorityState) const
+{
+	const FChaosMoverLaunchInputs& TypedAuthority = static_cast<const FChaosMoverLaunchInputs&>(AuthorityState);
+	return Mode != TypedAuthority.Mode || LaunchVelocityOrImpulse != TypedAuthority.LaunchVelocityOrImpulse;
+}
+
+void FChaosMoverLaunchInputs::Interpolate(const FMoverDataStructBase& From, const FMoverDataStructBase& To, float Pct)
+{
+	if (Pct < .5f)
+	{
+		*this = static_cast<const FChaosMoverLaunchInputs&>(From);
+	}
+	else
+	{
+		*this = static_cast<const FChaosMoverLaunchInputs&>(To);
+	}
+}
+
+void FChaosMoverLaunchInputs::Merge(const FMoverDataStructBase& From)
+{
+
+}
\ No newline at end of file
diff --git a/Engine/Plugins/Experimental/ChaosMover/Source/ChaosMover/Private/Character/ChaosCharacterMoverComponent.cpp b/Engine/Plugins/Experimental/ChaosMover/Source/ChaosMover/Private/Character/ChaosCharacterMoverComponent.cpp
--- a/Engine/Plugins/Experimental/ChaosMover/Source/ChaosMover/Private/Character/ChaosCharacterMoverComponent.cpp
+++ b/Engine/Plugins/Experimental/ChaosMover/Source/ChaosMover/Private/Character/ChaosCharacterMoverComponent.cpp
@@ -3,6 +3,7 @@
 #include "ChaosMover/Character/ChaosCharacterMoverComponent.h"
 
 #include "ChaosMover/Backends/ChaosMoverBackend.h"
+#include "ChaosMover/Character/ChaosCharacterInputs.h"
 #include "ChaosMover/Character/Modes/ChaosFallingMode.h"
 #include "ChaosMover/Character/Modes/ChaosFlyingMode.h"
 #include "ChaosMover/Character/Modes/ChaosWalkingMode.h"
@@ -56,4 +57,27 @@ bool UChaosCharacterMoverComponent::TryGetFloorCheckHitResult(FHitResult& OutHit
 	{
 		return Super::TryGetFloorCheckHitResult(OutHitResult);
 	}
+}
+
+void UChaosCharacterMoverComponent::ProduceInput(const int32 DeltaTimeMS, FMoverInputCmdContext* Cmd)
+{
+	Super::ProduceInput(DeltaTimeMS, Cmd);
+
+	// Add launch input data if launch velocity is set
+	if (!LaunchVelocityOrImpulse.IsZero())
+	{
+		check(Cmd);
+
+		FChaosMoverLaunchInputs& LaunchInputs = Cmd->InputCollection.FindOrAddMutableDataByType<FChaosMoverLaunchInputs>();
+		LaunchInputs.LaunchVelocityOrImpulse = LaunchVelocityOrImpulse;
+		LaunchInputs.Mode = LaunchMode;
+
+		LaunchVelocityOrImpulse = FVector::ZeroVector;
+	}
+}
+
+void UChaosCharacterMoverComponent::Launch(const FVector& VelocityOrImpulse, EChaosMoverVelocityEffectMode Mode)
+{
+	LaunchVelocityOrImpulse = VelocityOrImpulse;
+	LaunchMode = Mode;
 }
\ No newline at end of file
diff --git a/Engine/Plugins/Experimental/ChaosMover/Source/ChaosMover/Private/Character/Transitions/ChaosCharacterLaunchCheck.cpp b/Engine/Plugins/Experimental/ChaosMover/Source/ChaosMover/Private/Character/Transitions/ChaosCharacterLaunchCheck.cpp
new file mode 100644
--- /dev/null
+++ b/Engine/Plugins/Experimental/ChaosMover/Source/ChaosMover/Private/Character/Transitions/ChaosCharacterLaunchCheck.cpp
@@ -0,0 +1,59 @@
+// Copyright Epic Games, Inc. All Rights Reserved.
+
+#include "ChaosMover/Character/Transitions/ChaosCharacterLaunchCheck.h"
+
+#include "MoverComponent.h"
+#include "MoveLibrary/FloorQueryUtils.h"
+#include "ChaosMover/Character/Effects/ChaosCharacterApplyVelocityEffect.h"
+#include "ChaosMover/Character/ChaosCharacterInputs.h"
+#include "ChaosMover/ChaosMoverSimulation.h"
+#include "ChaosMover/ChaosMoverLog.h"
+#include "Chaos/PhysicsObjectInternalInterface.h"
+#include "Chaos/PhysicsObject.h"
+
+
+UChaosCharacterLaunchCheck::UChaosCharacterLaunchCheck(const FObjectInitializer& ObjectInitializer)
+	: Super(ObjectInitializer)
+{
+	bSupportsAsync = true;
+	bAllowModeReentry = true;
+	bFirstSubStepOnly = true;
+	TransitionToMode = DefaultModeNames::Falling;
+}
+
+FTransitionEvalResult UChaosCharacterLaunchCheck::Evaluate_Implementation(const FSimulationTickParams& Params) const
+{
+	FTransitionEvalResult EvalResult;
+
+	if (const FChaosMoverLaunchInputs* LaunchInputs = Params.StartState.InputCmd.InputCollection.FindDataByType<FChaosMoverLaunchInputs>())
+	{
+		if (TransitionToMode.IsNone())
+		{
+			EvalResult.NextMode = Params.StartState.SyncState.MovementMode;
+		}
+		else
+		{
+			EvalResult.NextMode = TransitionToMode;
+		}
+	}
+
+	return EvalResult;
+}
+
+void UChaosCharacterLaunchCheck::Trigger_Implementation(const FSimulationTickParams& Params)
+{
+	if (!Simulation)
+	{
+		UE_LOG(LogChaosMover, Warning, TEXT("No Simulation set on UChaosCharacterLaunchCheck"));
+		return;
+	}
+
+	if (const FChaosMoverLaunchInputs* LaunchInputs = Params.StartState.InputCmd.InputCollection.FindDataByType<FChaosMoverLaunchInputs>())
+	{
+		TSharedPtr<FChaosCharacterApplyVelocityEffect> LaunchMove = MakeShared<FChaosCharacterApplyVelocityEffect>();
+		LaunchMove->VelocityOrImpulseToApply = LaunchInputs->LaunchVelocityOrImpulse;
+		LaunchMove->Mode = LaunchInputs->Mode;
+
+		Simulation->QueueInstantMovementEffect(LaunchMove);
+	}
+}
diff --git a/Engine/Plugins/Experimental/ChaosMover/Source/ChaosMover/Public/ChaosMover/ChaosMoverSimulationTypes.h b/Engine/Plugins/Experimental/ChaosMover/Source/ChaosMover/Public/ChaosMover/ChaosMoverSimulationTypes.h
--- a/Engine/Plugins/Experimental/ChaosMover/Source/ChaosMover/Public/ChaosMover/ChaosMoverSimulationTypes.h
+++ b/Engine/Plugins/Experimental/ChaosMover/Source/ChaosMover/Public/ChaosMover/ChaosMoverSimulationTypes.h
@@ -177,3 +177,16 @@ public:
 	// Constraint settings of the joint constraint, only valid when ShouldUseConstraint() is true
 	virtual const Chaos::FPBDJointSettings& GetConstraintSettings() const = 0;
 };
+
+UENUM(BlueprintType)
+enum class EChaosMoverVelocityEffectMode : uint8
+{
+	/** Apply as an additive impulse*/
+	Impulse,
+
+	/** Apply as an additive velocity */
+	AdditiveVelocity,
+
+	/** Apply as an override velocity */
+	OverrideVelocity,
+};
diff --git a/Engine/Plugins/Experimental/ChaosMover/Source/ChaosMover/Public/ChaosMover/Character/ChaosCharacterInputs.h b/Engine/Plugins/Experimental/ChaosMover/Source/ChaosMover/Public/ChaosMover/Character/ChaosCharacterInputs.h
new file mode 100644
--- /dev/null
+++ b/Engine/Plugins/Experimental/ChaosMover/Source/ChaosMover/Public/ChaosMover/Character/ChaosCharacterInputs.h
@@ -0,0 +1,46 @@
+// Copyright Epic Games, Inc. All Rights Reserved.
+
+#pragma once
+
+#include "ChaosMover/Character/Effects/ChaosCharacterApplyVelocityEffect.h"
+#include "MoverTypes.h"
+
+#include "ChaosCharacterInputs.generated.h"
+
+USTRUCT(BlueprintType)
+struct FChaosMoverLaunchInputs : public FMoverDataStructBase
+{
+	GENERATED_BODY()
+
+public:
+	// Launch velocity in cm/s or impulse in kg cm/s
+	UPROPERTY(BlueprintReadWrite, Category = Mover)
+	FVector LaunchVelocityOrImpulse = FVector::ZeroVector;
+
+	UPROPERTY(BlueprintReadWrite, Category = Mover)
+	EChaosMoverVelocityEffectMode Mode = EChaosMoverVelocityEffectMode::AdditiveVelocity;
+
+	virtual UScriptStruct* GetScriptStruct() const override
+	{
+		return StaticStruct();
+	}
+	virtual FMoverDataStructBase* Clone() const override
+	{
+		return new FChaosMoverLaunchInputs(*this);
+	}
+	CHAOSMOVER_API virtual bool NetSerialize(FArchive& Ar, UPackageMap* Map, bool& bOutSuccess) override;
+	CHAOSMOVER_API virtual void ToString(FAnsiStringBuilderBase& Out) const override;
+	CHAOSMOVER_API virtual bool ShouldReconcile(const FMoverDataStructBase& AuthorityState) const override;
+	CHAOSMOVER_API virtual void Interpolate(const FMoverDataStructBase& From, const FMoverDataStructBase& To, float Pct) override;
+	CHAOSMOVER_API virtual void Merge(const FMoverDataStructBase& From) override;
+};
+
+template<>
+struct TStructOpsTypeTraits< FChaosMoverLaunchInputs > : public TStructOpsTypeTraitsBase2< FChaosMoverLaunchInputs >
+{
+	enum
+	{
+		WithNetSerializer = true,
+		WithCopy = true
+	};
+};
\ No newline at end of file
diff --git a/Engine/Plugins/Experimental/ChaosMover/Source/ChaosMover/Public/ChaosMover/Character/ChaosCharacterMoverComponent.h b/Engine/Plugins/Experimental/ChaosMover/Source/ChaosMover/Public/ChaosMover/Character/ChaosCharacterMoverComponent.h
--- a/Engine/Plugins/Experimental/ChaosMover/Source/ChaosMover/Public/ChaosMover/Character/ChaosCharacterMoverComponent.h
+++ b/Engine/Plugins/Experimental/ChaosMover/Source/ChaosMover/Public/ChaosMover/Character/ChaosCharacterMoverComponent.h
@@ -2,6 +2,7 @@
 
 #pragma once
 
+#include "ChaosMover/ChaosMoverSimulationTypes.h"
 #include "DefaultMovementSet/CharacterMoverComponent.h"
 #include "MoveLibrary/FloorQueryUtils.h"
 
@@ -20,15 +21,24 @@ public:
 	CHAOSMOVER_API UChaosCharacterMoverComponent();
 
 	CHAOSMOVER_API virtual bool TryGetFloorCheckHitResult(FHitResult& OutHitResult) const override;
+	CHAOSMOVER_API virtual void ProduceInput(const int32 DeltaTimeMS, FMoverInputCmdContext* Cmd) override;
 
 	// Broadcast when this actor lands on a valid surface.
-	UPROPERTY(BlueprintAssignable, Category = Mover)
+	UPROPERTY(BlueprintAssignable, Category = "Chaos Mover")
 	FChaosMover_OnLanded OnLanded;
 
+	// Launch the character using either impulse or velocity
+	// Note: This will only trigger a launch if a launch transition is implemented on the current movement mode
+	UFUNCTION(BlueprintCallable, Category = "Chaos Mover")
+	void Launch(const FVector& VelocityOrImpulse, EChaosMoverVelocityEffectMode Mode = EChaosMoverVelocityEffectMode::AdditiveVelocity);
+
 protected:
 	CHAOSMOVER_API virtual void ProcessSimulationEvent(const FMoverSimulationEventData& EventData) override;
 	CHAOSMOVER_API virtual void SetAdditionalSimulationOutput(const FMoverDataCollection& Data) override;
 
 	bool bFloorResultSet = false;
 	FFloorCheckResult LatestFloorResult;
+
+	FVector LaunchVelocityOrImpulse = FVector::ZeroVector;
+	EChaosMoverVelocityEffectMode LaunchMode = EChaosMoverVelocityEffectMode::AdditiveVelocity;
 };
diff --git a/Engine/Plugins/Experimental/ChaosMover/Source/ChaosMover/Public/ChaosMover/Character/Effects/ChaosCharacterApplyVelocityEffect.h b/Engine/Plugins/Experimental/ChaosMover/Source/ChaosMover/Public/ChaosMover/Character/Effects/ChaosCharacterApplyVelocityEffect.h
--- a/Engine/Plugins/Experimental/ChaosMover/Source/ChaosMover/Public/ChaosMover/Character/Effects/ChaosCharacterApplyVelocityEffect.h
+++ b/Engine/Plugins/Experimental/ChaosMover/Source/ChaosMover/Public/ChaosMover/Character/Effects/ChaosCharacterApplyVelocityEffect.h
@@ -2,22 +2,11 @@
 
 #pragma once
 
+#include "ChaosMover/ChaosMoverSimulationTypes.h"
 #include "InstantMovementEffect.h"
 
 #include "ChaosCharacterApplyVelocityEffect.generated.h"
 
-UENUM(BlueprintType)
-enum class EChaosMoverVelocityEffectMode : uint8
-{
-	/** Apply as an additive impulse*/
-	Impulse,
-
-	/** Apply as an additive velocity */
-	AdditiveVelocity,
-
-	/** Apply as an override velocity */
-	OverrideVelocity,
-};
 
 /** Applies a velocity or impulse for a single tick */
 USTRUCT(BlueprintInternalUseOnly)
@@ -26,7 +15,7 @@ struct CHAOSMOVER_API FChaosCharacterApplyVelocityEffect : public FInstantMoveme
 	GENERATED_BODY()
 
 	// Velocity or impulse to apply to the actor.
-	UPROPERTY(EditAnywhere, BlueprintReadWrite, Category = Settings, meta = (ForceUnits = "cm/s"))
+	UPROPERTY(EditAnywhere, BlueprintReadWrite, Category = Settings)
 	FVector VelocityOrImpulseToApply = FVector::ZeroVector;
 
 	// Controls whether to apply velocity or impulse and if the velocity will be additive
diff --git a/Engine/Plugins/Experimental/ChaosMover/Source/ChaosMover/Public/ChaosMover/Character/Transitions/ChaosCharacterLaunchCheck.h b/Engine/Plugins/Experimental/ChaosMover/Source/ChaosMover/Public/ChaosMover/Character/Transitions/ChaosCharacterLaunchCheck.h
new file mode 100644
--- /dev/null
+++ b/Engine/Plugins/Experimental/ChaosMover/Source/ChaosMover/Public/ChaosMover/Character/Transitions/ChaosCharacterLaunchCheck.h
@@ -0,0 +1,29 @@
+// Copyright Epic Games, Inc. All Rights Reserved.
+
+#pragma once
+
+#include "ChaosMover/ChaosMovementModeTransition.h"
+
+#include "ChaosCharacterLaunchCheck.generated.h"
+
+/**
+ * Transition that handles launching based on input for a physics-based character
+ */
+UCLASS(MinimalAPI, Blueprintable, BlueprintType)
+class UChaosCharacterLaunchCheck : public UChaosMovementModeTransition
+{
+	GENERATED_BODY()
+
+public:
+	CHAOSMOVER_API UChaosCharacterLaunchCheck(const class FObjectInitializer& ObjectInitializer);
+
+	CHAOSMOVER_API virtual FTransitionEvalResult Evaluate_Implementation(const FSimulationTickParams& Params) const override;
+	CHAOSMOVER_API virtual void Trigger_Implementation(const FSimulationTickParams& Params) override;
+
+	/**
+	Optional: Name of movement mode to transition to when launch is activated.
+	If set to None launching will not change the current mode
+	*/
+	UPROPERTY(BlueprintReadWrite, EditAnywhere, Category = Parameters)
+	FName TransitionToMode;
+};
diff --git a/Engine/Plugins/Experimental/Mover/Source/Mover/Public/MoverComponent.h b/Engine/Plugins/Experimental/Mover/Source/Mover/Public/MoverComponent.h
--- a/Engine/Plugins/Experimental/Mover/Source/Mover/Public/MoverComponent.h
+++ b/Engine/Plugins/Experimental/Mover/Source/Mover/Public/MoverComponent.h
@@ -121,7 +121,7 @@ public:
 	// --------------------------------------------------------------------------------
 
 	// Get latest local input prior to simulation step. Called by backend system on owner's instance (autonomous or authority).
-	MOVER_API void ProduceInput(const int32 DeltaTimeMS, FMoverInputCmdContext* Cmd);
+	MOVER_API virtual void ProduceInput(const int32 DeltaTimeMS, FMoverInputCmdContext* Cmd);
 
 	// Restore a previous frame prior to resimulating. Called by backend system.
 	MOVER_API void RestoreFrame(const FMoverSyncState* SyncState, const FMoverAuxStateContext* AuxState);
