Chaos Mover: Fixing floating point precision error in ChaosMoverStateMachines's fixed delta time calculation

diff --git a/Engine/Plugins/Experimental/ChaosMover/Source/ChaosMover/Private/ChaosMoverStateMachine.cpp b/Engine/Plugins/Experimental/ChaosMover/Source/ChaosMover/Private/ChaosMoverStateMachine.cpp
--- a/Engine/Plugins/Experimental/ChaosMover/Source/ChaosMover/Private/ChaosMoverStateMachine.cpp
+++ b/Engine/Plugins/Experimental/ChaosMover/Source/ChaosMover/Private/ChaosMoverStateMachine.cpp
@@ -173,7 +173,6 @@ namespace UE::ChaosMover
 			return;
 		}
 
-		const float EndingSimTimeMs = TimeStep.BaseSimTimeMs + TimeStep.StepMs;
 		InternalSimTimeMs = TimeStep.BaseSimTimeMs;
 
 		//GEngine->AddOnScreenDebugMessage(-1, -0.1f, FColor::White, FString::Printf(TEXT("Mode Tick: %s  Queued: %s"), *CurrentModeName.ToString(), *NextModeName.ToString()));
@@ -193,7 +192,9 @@ namespace UE::ChaosMover
 		int SubStepCount = 0;
 		const int32 MaxConsecutiveFullRefundedSubsteps = GetDefault<UMoverDeveloperSettings>()->MaxTimesToRefundSubstep;
 		int32 NumConsecutiveFullRefundedSubsteps = 0;
-		while (SubTimeStep.BaseSimTimeMs < EndingSimTimeMs)
+
+		float TotalUsedMs = 0.0f;
+		while (TotalUsedMs < TimeStep.StepMs)
 		{
 			InternalSimTimeMs = SubTimeStep.BaseSimTimeMs;
 
@@ -203,7 +204,7 @@ namespace UE::ChaosMover
 
 			OutputState.MovementEndState.ResetToDefaults();
 
-			SubTimeStep.StepMs = EndingSimTimeMs - SubTimeStep.BaseSimTimeMs;		// TODO: convert this to an overridable function that can support MaxStepTime, MaxIterations, etc.
+			SubTimeStep.StepMs = TimeStep.StepMs - TotalUsedMs;		// TODO: convert this to an overridable function that can support MaxStepTime, MaxIterations, etc.
 
 			// Transfer any queued moves into the starting state. They'll be started during the move generation.
 			FlushQueuedMovesToGroup(SubstepStartData.SyncState.LayeredMoves);
@@ -387,7 +388,7 @@ namespace UE::ChaosMover
 							*CurrentModeName.ToString(),
 							*OutputState.MovementEndState.NextModeName.ToString(),
 							*OwnerActorName);
-						SubTimeStep.BaseSimTimeMs += SubTimeStep.StepMs;
+						TotalUsedMs += SubTimeStep.StepMs;
 					}
 				}
 				else
@@ -405,8 +406,10 @@ namespace UE::ChaosMover
 			ensureMsgf(ActiveModifiers.IsEmpty(), TEXT("Movement Modifiers are not supported in Async mode"));
 
 			const float RemainingMs = FMath::Clamp(OutputState.MovementEndState.RemainingMs, 0.0f, SubTimeStep.StepMs);
-			SubTimeStep.BaseSimTimeMs += (SubTimeStep.StepMs - RemainingMs);
-			SubTimeStep.StepMs = EndingSimTimeMs - SubTimeStep.BaseSimTimeMs;
+			const float SubstepUsedMs = (SubTimeStep.StepMs - RemainingMs);
+			SubTimeStep.BaseSimTimeMs += SubstepUsedMs;
+			TotalUsedMs += SubstepUsedMs;
+			SubTimeStep.StepMs = RemainingMs;
 
 			SubstepStartData.SyncState = OutputState.SyncState;
 			SubstepStartData.AuxState = OutputState.AuxState;
@@ -414,7 +417,7 @@ namespace UE::ChaosMover
 			++SubStepCount;
 		}
 
-		InternalSimTimeMs = EndingSimTimeMs;
+		InternalSimTimeMs = TimeStep.BaseSimTimeMs + TotalUsedMs;
 
 		// Apply any instant effects that were queued up during this tick and didn't get handled in a substep
 		if (ApplyInstantEffects(SubstepStartData, SubTimeStep, OutputState.SyncState))
