Chaos Mover - Fix possible concurrent access of Simulation InputCmd

#rnx

diff --git a/Engine/Plugins/Experimental/ChaosMover/Source/ChaosMover/Private/Backends/ChaosMoverAsyncCallback.cpp b/Engine/Plugins/Experimental/ChaosMover/Source/ChaosMover/Private/Backends/ChaosMoverAsyncCallback.cpp
--- a/Engine/Plugins/Experimental/ChaosMover/Source/ChaosMover/Private/Backends/ChaosMoverAsyncCallback.cpp
+++ b/Engine/Plugins/Experimental/ChaosMover/Source/ChaosMover/Private/Backends/ChaosMoverAsyncCallback.cpp
@@ -11,6 +11,26 @@
 
 namespace UE::ChaosMover
 {
+	void FAsyncCallback::ProcessInputs_Internal(int32 PhysicsStep)
+	{
+		const FAsyncCallbackInput* AsyncInput = GetConsumerInput_Internal();
+
+		if (!AsyncInput || AsyncInput->InputData.IsEmpty())
+		{
+			return;
+		}
+
+		auto LambdaParallelUpdate = [&AsyncInput, PhysicsStep](int32 Idx) {
+			if (TStrongObjectPtr<UChaosMoverBackendComponent> Backend = AsyncInput->Backends[Idx].Pin())
+			{
+				Backend->GetSimulation()->ProcessInputs(PhysicsStep, AsyncInput->InputData[Idx]);
+			}
+			};
+
+		const bool ForceSingleThread = UE::ChaosMover::CVars::bForceSingleThreadedPT;
+		Chaos::PhysicsParallelFor(AsyncInput->Backends.Num(), LambdaParallelUpdate, ForceSingleThread);
+	}
+
 	void FAsyncCallback::OnPreSimulate_Internal()
 	{
 		const FAsyncCallbackInput* AsyncInput = GetConsumerInput_Internal();
diff --git a/Engine/Plugins/Experimental/ChaosMover/Source/ChaosMover/Private/Backends/ChaosMoverAsyncCallback.h b/Engine/Plugins/Experimental/ChaosMover/Source/ChaosMover/Private/Backends/ChaosMoverAsyncCallback.h
--- a/Engine/Plugins/Experimental/ChaosMover/Source/ChaosMover/Private/Backends/ChaosMoverAsyncCallback.h
+++ b/Engine/Plugins/Experimental/ChaosMover/Source/ChaosMover/Private/Backends/ChaosMoverAsyncCallback.h
@@ -46,6 +46,7 @@ namespace UE::ChaosMover
 		Chaos::ESimCallbackOptions::Rewind>
 	{
 	protected:
+		virtual void ProcessInputs_Internal(int32 PhysicsStep) override;
 		virtual void OnPreSimulate_Internal() override;
 		virtual void OnContactModification_Internal(Chaos::FCollisionContactModifier& Modifier) override;
 	};
diff --git a/Engine/Plugins/Experimental/ChaosMover/Source/ChaosMover/Private/Backends/ChaosMoverBackend.cpp b/Engine/Plugins/Experimental/ChaosMover/Source/ChaosMover/Private/Backends/ChaosMoverBackend.cpp
--- a/Engine/Plugins/Experimental/ChaosMover/Source/ChaosMover/Private/Backends/ChaosMoverBackend.cpp
+++ b/Engine/Plugins/Experimental/ChaosMover/Source/ChaosMover/Private/Backends/ChaosMoverBackend.cpp
@@ -395,10 +395,6 @@ void UChaosMoverBackendComponent::ProduceInputData(int32 PhysicsStep, int32 NumS
 		GenerateInput(TimeStep, InputData);
 	}
 
-	// Cache the produced input on the simulation so that it can be written to the network data
-	// This happens before the async input is received
-	Simulation->InitNetInputData(InputData.InputCmd);
-
 	UMoverComponent& MoverComp = GetMoverComponent();
 
 	// Add default simulation input data
diff --git a/Engine/Plugins/Experimental/ChaosMover/Source/ChaosMover/Private/ChaosMoverSimulation.cpp b/Engine/Plugins/Experimental/ChaosMover/Source/ChaosMover/Private/ChaosMoverSimulation.cpp
--- a/Engine/Plugins/Experimental/ChaosMover/Source/ChaosMover/Private/ChaosMoverSimulation.cpp
+++ b/Engine/Plugins/Experimental/ChaosMover/Source/ChaosMover/Private/ChaosMoverSimulation.cpp
@@ -67,20 +67,15 @@ const UBaseMovementMode* UChaosMoverSimulation::FindMovementModeByName(const FNa
 	return StateMachine.FindMovementMode(Name).Get();
 }
 
-void UChaosMoverSimulation::InitNetInputData(const FMoverInputCmdContext& InNetInputCmd)
-{
-	InputCmd = InNetInputCmd;
-}
-
 void UChaosMoverSimulation::ApplyNetInputData(const FMoverInputCmdContext& InNetInputCmd)
 {
-	InputCmd = InNetInputCmd;
+	NetInputCmd = InNetInputCmd;
 	bInputCmdOverridden = true;
 }
 
 void UChaosMoverSimulation::BuildNetInputData(FMoverInputCmdContext& OutNetInputCmd) const
 {
-	OutNetInputCmd = InputCmd;
+	OutNetInputCmd = NetInputCmd;
 }
 
 void UChaosMoverSimulation::ApplyNetStateData(const FMoverSyncState& InNetSyncState)
@@ -164,6 +159,24 @@ void UChaosMoverSimulation::OnDeinit()
 
 }
 
+void UChaosMoverSimulation::ProcessInputs(int32 PhysicsStep, const UE::ChaosMover::FSimulationInputData& InputData)
+{
+	Chaos::EnsureIsInPhysicsThreadContext();
+
+	// Ensure that we're using the correct inputs. On the server or during resim the inputs will
+	// be overridden by the network physics component, so use them instead of the generated inputs.
+	// On the client the input data is read out to be sent to the server so update the data that is read.
+	if (bInputCmdOverridden)
+	{
+		InputData.InputCmd = NetInputCmd;
+		bInputCmdOverridden = false;
+	}
+	else
+	{
+		NetInputCmd = InputData.InputCmd;
+	}
+}
+
 void UChaosMoverSimulation::SimulationTick(const FMoverTimeStep& TimeStep, const UE::ChaosMover::FSimulationInputData& InputData, UE::ChaosMover::FSimulationOutputData& OutputData)
 {
 	Chaos::EnsureIsInPhysicsThreadContext();
@@ -254,11 +267,6 @@ void UChaosMoverSimulation::PreSimulationTickCharacter(const IChaosCharacterMove
 
 void UChaosMoverSimulation::OnPreSimulationTick(const FMoverTimeStep& TimeStep, const UE::ChaosMover::FSimulationInputData& InputData)
 {
-	if (bInputCmdOverridden)
-	{
-		InputData.InputCmd = InputCmd;
-	}
-
 	if (TimeStep.bIsResimulating)
 	{
 		check(Solver);
diff --git a/Engine/Plugins/Experimental/ChaosMover/Source/ChaosMover/Public/ChaosMover/ChaosMoverSimulation.h b/Engine/Plugins/Experimental/ChaosMover/Source/ChaosMover/Public/ChaosMover/ChaosMoverSimulation.h
--- a/Engine/Plugins/Experimental/ChaosMover/Source/ChaosMover/Public/ChaosMover/ChaosMoverSimulation.h
+++ b/Engine/Plugins/Experimental/ChaosMover/Source/ChaosMover/Public/ChaosMover/ChaosMoverSimulation.h
@@ -69,12 +69,13 @@ public:
 
 	CHAOSMOVER_API void Init(const FInitParams& InitParams);
 	CHAOSMOVER_API void Deinit();
+
+	CHAOSMOVER_API void ProcessInputs(int32 PhysicsStep, const UE::ChaosMover::FSimulationInputData& InputData);
 	CHAOSMOVER_API void SimulationTick(const FMoverTimeStep& TimeStep, const UE::ChaosMover::FSimulationInputData& InputData, UE::ChaosMover::FSimulationOutputData& OutputData);
 	CHAOSMOVER_API void ModifyContacts(const FMoverTimeStep& TimeStep, const UE::ChaosMover::FSimulationInputData& InputData, const UE::ChaosMover::FSimulationOutputData& OutputData, Chaos::FCollisionContactModifier& Modifier) const;
 
 	CHAOSMOVER_API void AddEvent(TSharedPtr<FMoverSimulationEventData> Event);
 
-	CHAOSMOVER_API void InitNetInputData(const FMoverInputCmdContext& InNetInputCmd);
 	CHAOSMOVER_API void ApplyNetInputData(const FMoverInputCmdContext& InNetInputCmd);
 	CHAOSMOVER_API void BuildNetInputData(FMoverInputCmdContext& OutNetInputCmd) const;
 	CHAOSMOVER_API void ApplyNetStateData(const FMoverSyncState& InNetSyncState);
@@ -153,7 +154,7 @@ protected:
 	UWorld* World = nullptr;
 
 private:
-	FMoverInputCmdContext InputCmd;
+	FMoverInputCmdContext NetInputCmd;
 	TArray<TSharedPtr<FMoverSimulationEventData>> Events;
 
 	bool bInputCmdOverridden = false;
