Mover:
- Support for more than one input producer. This allows us to support several mode type controllers, and therefore truly mix different movement modes, like character movement modes, and pathed movement modes, without having to derive from UMoverComponent or supporting every possible modes at the UMoverComponent level
- Generic simulation event dispatch (C++ only, to expose to BP, dedicated delegates should be added for each event type)

Chaos Mover
- "Path Target Constraint" renamed to general purpose "Actuation Constraint" as it could be used for movement that is not on a path
- Added ChaosMoverSettings for in editor exposed settings. We might want to rename the ChaosMoverDevelopperSettings.h file something else
as it does not containt "UDeveloperSettings" but rather internal cvars.
- Hooking up OnPostFinalize in ChaosMoverBackend so that chaos mover component can chave access to the post sim sync state. This was needed to implement pathed movement correctly.
- Hooking up QueueNextMode to the Chaos Mover simulation so it gets affected by mode change requests coming from gameplay. This is implemented at the ChaosCharacterMoverComponent level
only, but maybe this should get moved up to UMoverComponent
- Pathed Movement
- Port of Dan Hertzka's PathedMovementMode to ChaosMover
- Added a ChaosPathedMovementControllerComponent to produce input for pathed movement modes and expose pathed movement mode state and events back to gameplay code and BP
- PathedMovementReachedEndTransition fires when pathed movement mode has finished

MoverExamples:
- Added L_PathedMovement.umap, a test map for pathed movement
- Added a PathedMovementMode to ChaosMannyPawn to showcase that pathed movement mode can exist on a character

[RN]

#virtualized

diff --git a/Engine/Plugins/Experimental/ChaosMover/Source/ChaosMover/Private/Backends/ChaosMoverAsyncCallback.cpp b/Engine/Plugins/Experimental/ChaosMover/Source/ChaosMover/Private/Backends/ChaosMoverAsyncCallback.cpp
--- a/Engine/Plugins/Experimental/ChaosMover/Source/ChaosMover/Private/Backends/ChaosMoverAsyncCallback.cpp
+++ b/Engine/Plugins/Experimental/ChaosMover/Source/ChaosMover/Private/Backends/ChaosMoverAsyncCallback.cpp
@@ -41,13 +41,7 @@ namespace UE::ChaosMover
 		}
 
 		FMoverTimeStep TimeStep;
-		TimeStep.BaseSimTimeMs = GetSimTime_Internal() * 1000.0f;
-		TimeStep.StepMs = GetDeltaTime_Internal() * 1000.0f;
-		if (AsyncInput->PhysicsSolver)
-		{
-			TimeStep.ServerFrame = AsyncInput->PhysicsSolver->GetCurrentFrame() + AsyncInput->NetworkPhysicsTickOffset;
-			TimeStep.bIsResimulating = AsyncInput->PhysicsSolver->GetEvolution()->IsResimming();
-		}
+		GeCurrentMoverTimeStep(AsyncInput, TimeStep);
 
 		FAsyncCallbackOutput& AsyncOutput = GetProducerOutputData_Internal();
 		AsyncOutput.OutputData.SetNum(AsyncInput->InputData.Num());
@@ -76,13 +70,7 @@ namespace UE::ChaosMover
 		}
 
 		FMoverTimeStep TimeStep;
-		TimeStep.BaseSimTimeMs = GetSimTime_Internal() * 1000.0f;
-		TimeStep.StepMs = GetDeltaTime_Internal() * 1000.0f;
-		if (AsyncInput->PhysicsSolver)
-		{
-			TimeStep.ServerFrame = AsyncInput->PhysicsSolver->GetCurrentFrame() + AsyncInput->NetworkPhysicsTickOffset;
-			TimeStep.bIsResimulating = AsyncInput->PhysicsSolver->GetEvolution()->IsResimming();
-		}
+		GeCurrentMoverTimeStep(AsyncInput, TimeStep);
 
 		const FAsyncCallbackOutput& AsyncOutput = GetProducerOutputData_Internal();
 		for (int Idx = 0; Idx < AsyncInput->Backends.Num(); ++Idx)
@@ -93,4 +81,15 @@ namespace UE::ChaosMover
 			}
 		}
 	}
+
+	void FAsyncCallback::GeCurrentMoverTimeStep(const FAsyncCallbackInput* AsyncInput, FMoverTimeStep& OutMoverTimeStep) const
+	{
+		OutMoverTimeStep.BaseSimTimeMs = GetSimTime_Internal() * 1000.0f;
+		OutMoverTimeStep.StepMs = GetDeltaTime_Internal() * 1000.0f;
+		if (AsyncInput && AsyncInput->PhysicsSolver)
+		{
+			OutMoverTimeStep.ServerFrame = AsyncInput->PhysicsSolver->GetCurrentFrame() + AsyncInput->NetworkPhysicsTickOffset;
+			OutMoverTimeStep.bIsResimulating = AsyncInput->PhysicsSolver->GetEvolution()->IsResimming();
+		}
+	}
 }
\ No newline at end of file
diff --git a/Engine/Plugins/Experimental/ChaosMover/Source/ChaosMover/Private/Backends/ChaosMoverAsyncCallback.h b/Engine/Plugins/Experimental/ChaosMover/Source/ChaosMover/Private/Backends/ChaosMoverAsyncCallback.h
--- a/Engine/Plugins/Experimental/ChaosMover/Source/ChaosMover/Private/Backends/ChaosMoverAsyncCallback.h
+++ b/Engine/Plugins/Experimental/ChaosMover/Source/ChaosMover/Private/Backends/ChaosMoverAsyncCallback.h
@@ -49,5 +49,7 @@ namespace UE::ChaosMover
 		virtual void ProcessInputs_Internal(int32 PhysicsStep) override;
 		virtual void OnPreSimulate_Internal() override;
 		virtual void OnContactModification_Internal(Chaos::FCollisionContactModifier& Modifier) override;
+
+		void GeCurrentMoverTimeStep(const FAsyncCallbackInput* AsyncInput, FMoverTimeStep& OutMoverTimeStep) const;
 	};
 }
\ No newline at end of file
diff --git a/Engine/Plugins/Experimental/ChaosMover/Source/ChaosMover/Private/Backends/ChaosMoverBackend.cpp b/Engine/Plugins/Experimental/ChaosMover/Source/ChaosMover/Private/Backends/ChaosMoverBackend.cpp
--- a/Engine/Plugins/Experimental/ChaosMover/Source/ChaosMover/Private/Backends/ChaosMoverBackend.cpp
+++ b/Engine/Plugins/Experimental/ChaosMover/Source/ChaosMover/Private/Backends/ChaosMoverBackend.cpp
@@ -24,7 +24,7 @@
 #include UE_INLINE_GENERATED_CPP_BY_NAME(ChaosMoverBackend)
 
 UChaosMoverBackendComponent::UChaosMoverBackendComponent()
-	: PathTargetConstraintPhysicsUserData(&PathTargetConstraintInstance)
+	: ActuationConstraintPhysicsUserData(&ActuationConstraintInstance)
 {
 	PrimaryComponentTick.bCanEverTick = false;
 
@@ -143,15 +143,15 @@ void UChaosMoverBackendComponent::CreatePhysics()
 	// Create all possible constraints...
 	// ... a character ground constraint, for constraint based character-like movement on ground
 	CreateCharacterGroundConstraint();
-	// ... a path target constraint, for constraint based pathed movement
-	CreatePathTargetConstraint();
+	// ... a general purpose actuation joint constraint, for example, for constraint based pathed movement
+	CreateActuationConstraint();
 }
 
 void UChaosMoverBackendComponent::DestroyPhysics()
 {
 	// Destroy all constraints
 	DestroyCharacterGroundConstraint();
-	DestroyPathTargetConstraint();
+	DestroyActuationConstraint();
 }
 
 void UChaosMoverBackendComponent::CreateCharacterGroundConstraint()
@@ -183,7 +183,7 @@ void UChaosMoverBackendComponent::DestroyCharacterGroundConstraint()
 	}
 }
 
-void UChaosMoverBackendComponent::CreatePathTargetConstraint()
+void UChaosMoverBackendComponent::CreateActuationConstraint()
 {
 	const UMoverComponent& MoverComp = GetMoverComponent();
 	if (Chaos::FPhysicsObject* PhysicsObject = GetPhysicsObject())
@@ -198,9 +198,9 @@ void UChaosMoverBackendComponent::CreatePathTargetConstraint()
 			if (Chaos::FJointConstraint* Constraint = static_cast<Chaos::FJointConstraint*>(Handle.Constraint))
 			{
 				// Since we didn't use the ConstraintInstance to actually create the constraint (it requires both bodies exist, see comment above), link everything up manually
-				PathTargetConstraintHandle = Handle;						
-				PathTargetConstraintInstance.ConstraintHandle = PathTargetConstraintHandle;
-				Constraint->SetUserData(&PathTargetConstraintPhysicsUserData/*has a (void*)FConstraintInstanceBase*/);
+				ActuationConstraintHandle = Handle;						
+				ActuationConstraintInstance.ConstraintHandle = ActuationConstraintHandle;
+				Constraint->SetUserData(&ActuationConstraintPhysicsUserData/*has a (void*)FConstraintInstanceBase*/);
 				bIsConstraintValid = true;
 
 				if (Chaos::FPBDRigidParticle* EndpointParticle = Constraint->GetPhysicsBodies()[1]->GetParticle<Chaos::EThreadContext::External>()->CastToRigidParticle())
@@ -218,11 +218,11 @@ void UChaosMoverBackendComponent::CreatePathTargetConstraint()
 	}
 }
 
-void UChaosMoverBackendComponent::DestroyPathTargetConstraint()
+void UChaosMoverBackendComponent::DestroyActuationConstraint()
 {
-	if (PathTargetConstraintHandle.IsValid())
+	if (ActuationConstraintHandle.IsValid())
 	{
-		FChaosEngineInterface::ReleaseConstraint(PathTargetConstraintHandle);
+		FChaosEngineInterface::ReleaseConstraint(ActuationConstraintHandle);
 	}
 }
 
@@ -298,8 +298,8 @@ void UChaosMoverBackendComponent::InitSimulation()
 	Params.NullMovementMode = TWeakObjectPtr<UNullMovementMode>(NullMovementMode.Get());
 	Params.StartingMovementMode = MoverComp.StartingMovementMode;
 	Params.CharacterConstraintProxy = CharacterGroundConstraint ? CharacterGroundConstraint->GetProxy<Chaos::FCharacterGroundConstraintProxy>() : nullptr;
-	Params.PathTargetConstraintProxy = PathTargetConstraintHandle.IsValid() ? PathTargetConstraintHandle->GetProxy<Chaos::FJointConstraintPhysicsProxy>() : nullptr;
-	Params.PathTargetKinematicEndPointProxy = Params.PathTargetConstraintProxy && Params.PathTargetConstraintProxy->GetConstraint() ? Params.PathTargetConstraintProxy->GetConstraint()->GetKinematicEndPoint() : nullptr;
+	Params.ActuationConstraintProxy = ActuationConstraintHandle.IsValid() ? ActuationConstraintHandle->GetProxy<Chaos::FJointConstraintPhysicsProxy>() : nullptr;
+	Params.ActuationConstraintKinematicEndPointProxy = Params.ActuationConstraintProxy && Params.ActuationConstraintProxy->GetConstraint() ? Params.ActuationConstraintProxy->GetConstraint()->GetKinematicEndPoint() : nullptr;
 	Params.PhysicsObject = GetPhysicsObject();
 	Params.Solver = GetPhysicsSolver();
 	Params.World = GetWorld();
@@ -371,6 +371,31 @@ void UChaosMoverBackendComponent::EndPlay(const EEndPlayReason::Type EndPlayReas
 	Super::EndPlay(EndPlayReason);
 }
 
+float UChaosMoverBackendComponent::GetCurrentSimTimeMs()
+{
+	return 0.0f;
+}
+
+int32 UChaosMoverBackendComponent::GetCurrentSimFrame()
+{
+	return 0;
+}
+
+bool UChaosMoverBackendComponent::IsAsync() const
+{
+	return true;
+}
+
+UChaosMoverSimulation* UChaosMoverBackendComponent::GetSimulation()
+{
+	return Simulation;
+}
+
+const UChaosMoverSimulation* UChaosMoverBackendComponent::GetSimulation() const
+{
+	return Simulation;
+}
+
 void UChaosMoverBackendComponent::ProduceInputData(int32 PhysicsStep, int32 NumSteps, const FMoverTimeStep& TimeStep, UE::ChaosMover::FSimulationInputData& InputData)
 {
 	Chaos::EnsureIsInGameThreadContext();
@@ -476,4 +501,10 @@ void UChaosMoverBackendComponent::FinalizeFrame(float ResultsTimeInMs)
 	{
 		MoverComp.OnPostSimulationTick.Broadcast(TimeStep);
 	}
+
+	if (MoverComp.OnPostFinalize.IsBound())
+	{
+		FMoverAuxStateContext UnusedAuxStateContext;
+		MoverComp.OnPostFinalize.Broadcast(InterpolatedOutput.SyncState, UnusedAuxStateContext);
+	}
 }
diff --git a/Engine/Plugins/Experimental/ChaosMover/Source/ChaosMover/Private/ChaosMoverSettings.cpp b/Engine/Plugins/Experimental/ChaosMover/Source/ChaosMover/Private/ChaosMoverSettings.cpp
new file mode 100644
--- /dev/null
+++ b/Engine/Plugins/Experimental/ChaosMover/Source/ChaosMover/Private/ChaosMoverSettings.cpp
@@ -0,0 +1,10 @@
+// Copyright Epic Games, Inc. All Rights Reserved.
+
+#include "ChaosMover/ChaosMoverSettings.h"
+
+#include UE_INLINE_GENERATED_CPP_BY_NAME(ChaosMoverSettings)
+
+UChaosMoverSettings::UChaosMoverSettings()
+{
+}
+
diff --git a/Engine/Plugins/Experimental/ChaosMover/Source/ChaosMover/Private/ChaosMoverSimulation.cpp b/Engine/Plugins/Experimental/ChaosMover/Source/ChaosMover/Private/ChaosMoverSimulation.cpp
--- a/Engine/Plugins/Experimental/ChaosMover/Source/ChaosMover/Private/ChaosMoverSimulation.cpp
+++ b/Engine/Plugins/Experimental/ChaosMover/Source/ChaosMover/Private/ChaosMoverSimulation.cpp
@@ -13,6 +13,7 @@
 #include "ChaosMover/ChaosMoverLog.h"
 #include "ChaosMover/ChaosMoverSimulationTypes.h"
 #include "ChaosMover/Utilities/ChaosMoverQueryUtils.h"
+#include "ChaosMover/PathedMovement/ChaosPathedMovementTypes.h"
 #include "Framework/Threading.h"
 #include "MoveLibrary/FloorQueryUtils.h"
 #include "MoveLibrary/MovementMixer.h"
@@ -95,8 +96,8 @@ void UChaosMoverSimulation::Init(const FInitParams& InitParams)
 
 	MovementMixerWeakPtr = InitParams.MovementMixer.Get();
 	CharacterConstraintProxy = InitParams.CharacterConstraintProxy;
-	PathTargetConstraintProxy = InitParams.PathTargetConstraintProxy;
-	PathTargetKinematicEndPointProxy = InitParams.PathTargetKinematicEndPointProxy;
+	ActuationConstraintProxy = InitParams.ActuationConstraintProxy;
+	ActuationConstraintKinematicEndPointProxy = InitParams.ActuationConstraintKinematicEndPointProxy;
 	PhysicsObject = InitParams.PhysicsObject;
 	Solver = InitParams.Solver;
 	World = InitParams.World;
@@ -109,6 +110,18 @@ void UChaosMoverSimulation::Init(const FInitParams& InitParams)
 	StateMachineInitParams.Simulation = this;
 	StateMachine.Init(StateMachineInitParams);
 
+	// Set the pathed movement basis to the initial location of the controlled particle
+	if (PhysicsObject)
+	{
+		if (Chaos::FGeometryParticle* Particle = FPhysicsObjectExternalInterface::LockRead(PhysicsObject)->GetRigidParticle(PhysicsObject))
+		{
+			FTransform NewPathBasis;
+			NewPathBasis.SetTranslation(Particle->GetX());
+			NewPathBasis.SetRotation(Particle->GetR());
+			SetPathedMovementBasis(NewPathBasis);
+		}
+	}
+
 	for (const TPair<FName, TWeakObjectPtr<UBaseMovementMode>>& Element : InitParams.ModesToRegister)
 	{
 		const FName& ModeName = Element.Key;
@@ -139,7 +152,7 @@ void UChaosMoverSimulation::Init(const FInitParams& InitParams)
 		StateMachine.RegisterGlobalTransition(Transition);
 	}
 
-	StateMachine.SetModeImmediately(StateMachine.GetDefaultModeName());
+	StateMachine.QueueNextMode(StateMachine.GetDefaultModeName());
 
 	OnInit();
 }
@@ -286,7 +299,8 @@ void UChaosMoverSimulation::OnPreSimulationTick(const FMoverTimeStep& TimeStep,
 
 	if (TStrongObjectPtr<UBaseMovementMode> CurrentModePtr = StateMachine.GetCurrentMode().Pin())
 	{
-		if (IChaosCharacterMovementModeInterface* CharacterMode = Cast<IChaosCharacterMovementModeInterface>(CurrentModePtr.Get()))
+		UBaseMovementMode* CurrentMode = CurrentModePtr.Get();
+		if (IChaosCharacterMovementModeInterface* CharacterMode = Cast<IChaosCharacterMovementModeInterface>(CurrentMode))
 		{
 			PreSimulationTickCharacter(*CharacterMode, TimeStep, InputData);
 		}
@@ -448,35 +462,37 @@ void UChaosMoverSimulation::PostSimulationTickCharacterConstraint(const IChaosCh
 	}
 }
 
-void UChaosMoverSimulation::PostSimulationTickPathedMovement(const IChaosPathedMovementModeInterface& ConstraintMode, const FMoverTimeStep& TimeStep, UE::ChaosMover::FSimulationOutputData& OutputData)
+void UChaosMoverSimulation::PostSimulationTickMovementActuation(const IChaosMovementActuationInterface& MovementActuationMode, const FMoverTimeStep& TimeStep, UE::ChaosMover::FSimulationOutputData& OutputData)
 {
 	Chaos::FPBDJointConstraintHandle* ConstraintHandle = nullptr;
-	if (PathTargetConstraintProxy && PathTargetConstraintProxy->IsInitialized())
+	if (ActuationConstraintProxy && ActuationConstraintProxy->IsInitialized())
 	{
-		ConstraintHandle = PathTargetConstraintProxy->GetHandle();
+		ConstraintHandle = ActuationConstraintProxy->GetHandle();
 	}
 	if (!ConstraintHandle)
 	{
 		return;
 	}
 
-	if (ConstraintMode.ShouldUseConstraint())
+	if (MovementActuationMode.ShouldUseConstraint())
 	{
 		if (!ConstraintHandle->IsEnabled())
 		{
-			EnablePathTargetConstraint();
+			EnableActuationConstraint();
 		}
 
 		if (!IsControlledParticleDynamic())
 		{
 			SetControlledParticleDynamic();
 		}
+
+		ConstraintHandle->SetSettings(MovementActuationMode.GetConstraintSettings());
 	}
-	else if (!ConstraintMode.ShouldUseConstraint())
+	else if (!MovementActuationMode.ShouldUseConstraint())
 	{
 		if (ConstraintHandle->IsEnabled())
 		{
-			DisablePathTargetConstraint();
+			DisableActuationConstraint();
 		}
 
 		if (!IsControlledParticleKinematic())
@@ -487,25 +503,66 @@ void UChaosMoverSimulation::PostSimulationTickPathedMovement(const IChaosPathedM
 
 	const FMoverDefaultSyncState* PostSimDefaultSyncState = OutputData.SyncState.SyncStateCollection.FindMutableDataByType<FMoverDefaultSyncState>();
 	check(PostSimDefaultSyncState);
-	check(Solver);
-	Chaos::FPBDRigidsEvolution& Evolution = *Solver->GetEvolution();
+	SetActuationTargetTransform(PostSimDefaultSyncState->GetTransform_WorldSpace());
+}
+
+void UChaosMoverSimulation::SetActuationTargetTransform(const FTransform& TargetTransform)
+{
+	Chaos::FPBDRigidsEvolution* Evolution = Solver ? Solver->GetEvolution() : nullptr;
+	if (!Evolution)
+	{
+		return;
+	}
 
-	// We always update the path target end point, even when the controlled particle is not constrained.
-	// That way we don't leave a particle behind and don't risk having it very far and whip lashing when re-enabling the constraint.
-	Chaos::FKinematicGeometryParticleHandle* EndpointParticleHandle = (ConstraintHandle && PathTargetKinematicEndPointProxy) ? PathTargetKinematicEndPointProxy->GetHandle_LowLevel()->CastToKinematicParticle() : nullptr;
+	Chaos::FPBDJointConstraintHandle* ConstraintHandle = ActuationConstraintProxy && ActuationConstraintProxy->IsInitialized() ? ActuationConstraintProxy->GetHandle() : nullptr;
+	Chaos::FKinematicGeometryParticleHandle* EndpointParticleHandle = (ConstraintHandle && ActuationConstraintKinematicEndPointProxy) ? ActuationConstraintKinematicEndPointProxy->GetHandle_LowLevel()->CastToKinematicParticle() : nullptr;
+	Chaos::FKinematicTarget KinematicPathTarget = Chaos::FKinematicTarget::MakePositionTarget(TargetTransform);
 	if (EndpointParticleHandle)
 	{	
-		Evolution.SetParticleKinematicTarget(EndpointParticleHandle, Chaos::FKinematicTarget::MakePositionTarget(PostSimDefaultSyncState->GetTransform_WorldSpace()));
+		Evolution->SetParticleKinematicTarget(EndpointParticleHandle, KinematicPathTarget);
 	}
 	if (IsControlledParticleKinematic())
 	{
 		if (Chaos::FPBDRigidParticleHandle* ParticleHandle = GetControlledParticle())
 		{
-			Evolution.SetParticleKinematicTarget(ParticleHandle, Chaos::FKinematicTarget::MakePositionTarget(PostSimDefaultSyncState->GetTransform_WorldSpace()));
+			Evolution->SetParticleKinematicTarget(ParticleHandle, KinematicPathTarget);
 		}
 	}
 }
 
+void UChaosMoverSimulation::TeleportActuationTarget(const FTransform& TargetTransform, bool AlsoTeleportControlledParticle /*= false*/)
+{
+	Chaos::FPBDRigidsEvolution* Evolution = Solver ? Solver->GetEvolution() : nullptr;
+	if (!Evolution)
+	{
+		return;
+	}
+
+	Chaos::FPBDJointConstraintHandle* ConstraintHandle = (ActuationConstraintProxy && ActuationConstraintProxy->IsInitialized()) ? ActuationConstraintProxy->GetHandle() : nullptr;
+	Chaos::FKinematicGeometryParticleHandle* EndpointParticleHandle = (ConstraintHandle) ? ActuationConstraintKinematicEndPointProxy->GetHandle_LowLevel()->CastToKinematicParticle() : nullptr;
+	if (EndpointParticleHandle)
+	{	
+		Evolution->SetParticleTransform(EndpointParticleHandle, TargetTransform.GetTranslation(), TargetTransform.GetRotation(), /*bIsTeleport=*/true);
+	}
+	if (AlsoTeleportControlledParticle)
+	{
+		if (Chaos::FPBDRigidParticleHandle* ParticleHandle = GetControlledParticle())
+		{
+			Evolution->SetParticleTransform(ParticleHandle, TargetTransform.GetTranslation(), TargetTransform.GetRotation(), /*bIsTeleport=*/true);
+		}
+	}
+}
+
+const FTransform& UChaosMoverSimulation::GetPathedMovementBasis() const
+{
+	return PathBasis;
+}
+
+void UChaosMoverSimulation::SetPathedMovementBasis(const FTransform& InPathedMovementBasis)
+{
+	PathBasis = InPathedMovementBasis;
+}
+
 void UChaosMoverSimulation::OnPostSimulationTick(const FMoverTimeStep& TimeStep, UE::ChaosMover::FSimulationOutputData& OutputData)
 {
 	// TODO - make this more extensible
@@ -525,13 +582,13 @@ void UChaosMoverSimulation::OnPostSimulationTick(const FMoverTimeStep& TimeStep,
 			DisableCharacterConstraint();
 		}
 
-		if (IChaosPathedMovementModeInterface* ConstraintMode = Cast<IChaosPathedMovementModeInterface>(CurrentModePtr.Get()))
+		if (IChaosMovementActuationInterface* MovementActuationMode = Cast<IChaosMovementActuationInterface>(CurrentModePtr.Get()))
 		{
-			PostSimulationTickPathedMovement(*ConstraintMode, TimeStep, OutputData);
+			PostSimulationTickMovementActuation(*MovementActuationMode, TimeStep, OutputData);
 		}
 		else
 		{
-			DisablePathTargetConstraint();
+			DisableActuationConstraint();
 		}
 	}
 
@@ -622,13 +679,60 @@ void UChaosMoverSimulation::OnMovementModeChanged(const FMovementModeChangedEven
 			}
 		}
 
-		if (IChaosPathedMovementModeInterface* NextPathTargetConstraintMode = Cast<IChaosPathedMovementModeInterface>(NextModePtr.Get()))
+		if (IChaosMovementActuationInterface* NextMovementActuationInterface = Cast<IChaosMovementActuationInterface>(NextModePtr.Get()))
 		{
-			bool IsPathTargetConstraintInitialized = (PathTargetConstraintProxy != nullptr) && PathTargetConstraintProxy->IsInitialized();
-			Chaos::FPBDJointConstraintHandle* ConstraintHandle = IsPathTargetConstraintInitialized ? PathTargetConstraintProxy->GetHandle() : nullptr;
+			// Reset the pathed movement state
+			FChaosPathedMovementState* PresimPathedMovementState = CurrentSyncState.SyncStateCollection.FindMutableDataByType<FChaosPathedMovementState>();
+			if (PresimPathedMovementState)
+			{
+				*PresimPathedMovementState = FChaosPathedMovementState();
+			}
+
+			Chaos::FPBDJointConstraintHandle* ConstraintHandle = (ActuationConstraintProxy && ActuationConstraintProxy->IsInitialized()) ? ActuationConstraintProxy->GetHandle() : nullptr;
 			if (ConstraintHandle)
 			{
-				ConstraintHandle->SetSettings(NextPathTargetConstraintMode->GetConstraintSettings());
+				if (NextMovementActuationInterface->ShouldUseConstraint())
+				{
+					// Enable the constraint
+					if (!ConstraintHandle->IsEnabled())
+					{
+						EnableActuationConstraint();
+					}
+					// Set the controlled particle dynamic
+					if (!IsControlledParticleDynamic())
+					{
+						SetControlledParticleDynamic();
+					}
+					// Apply movement actuation constraint settings
+					// Bug: if we didn't also call SetSettings every frame in UChaosMoverSimulation::PostSimulationTickMovementActuation,
+					// the settings wouldn't apply correctly from this call only if this mode is the first active mode (OnMovementModeChanged called on the very first OnSimulationTick)
+					ConstraintHandle->SetSettings(NextMovementActuationInterface->GetConstraintSettings());
+
+					// If the mode was also a pathed mode, let's teleport to the position on the path
+					// we're supposed to be at, using CurrentPathFraction
+					Chaos::FRigidTransform3 TargetLastFrame = Chaos::FRigidTransform3::Identity;
+					IChaosPathedMovementModeInterface* NextPathedMovementMode = Cast<IChaosPathedMovementModeInterface>(NextMovementActuationInterface);
+					if (NextPathedMovementMode)
+					{	
+						float PathFraction = PresimPathedMovementState ? PresimPathedMovementState->CurrentPathFraction : 0.0f;
+						// We resume the path at the reset progress
+						// (if not 0, CurrentPathFraction was either not reset properly when the movement mode changed, or deliberately set to a different value)
+						TargetLastFrame = NextPathedMovementMode->CalcTargetRelativeTransform(PathFraction);
+					}
+					Chaos::FRigidTransform3 WorldTransform = TargetLastFrame * PathBasis;
+					TeleportActuationTarget(WorldTransform);
+				}
+				else
+				{
+					if (ConstraintHandle->IsEnabled())
+					{
+						DisableActuationConstraint();
+					}
+					if (!IsControlledParticleKinematic())
+					{
+						SetControlledParticleKinematic();
+					}
+				}
 			}
 		}
 	}
@@ -708,22 +812,22 @@ void UChaosMoverSimulation::DisableCharacterConstraint()
 	}
 }
 
-void UChaosMoverSimulation::EnablePathTargetConstraint()
+void UChaosMoverSimulation::EnableActuationConstraint()
 {
-	if (PathTargetConstraintProxy && PathTargetConstraintProxy->IsInitialized())
+	if (ActuationConstraintProxy && ActuationConstraintProxy->IsInitialized())
 	{
-		if (Chaos::FPBDJointConstraintHandle* ConstraintHandle = PathTargetConstraintProxy->GetHandle())
+		if (Chaos::FPBDJointConstraintHandle* ConstraintHandle = ActuationConstraintProxy->GetHandle())
 		{
 			ConstraintHandle->SetConstraintEnabled(true);
 		}
 	}
 }
 
-void UChaosMoverSimulation::DisablePathTargetConstraint()
+void UChaosMoverSimulation::DisableActuationConstraint()
 {
-	if (PathTargetConstraintProxy && PathTargetConstraintProxy->IsInitialized())
+	if (ActuationConstraintProxy && ActuationConstraintProxy->IsInitialized())
 	{
-		if (Chaos::FPBDJointConstraintHandle* ConstraintHandle = PathTargetConstraintProxy->GetHandle())
+		if (Chaos::FPBDJointConstraintHandle* ConstraintHandle = ActuationConstraintProxy->GetHandle())
 		{
 			ConstraintHandle->SetConstraintEnabled(false);
 		}
diff --git a/Engine/Plugins/Experimental/ChaosMover/Source/ChaosMover/Private/Character/ChaosCharacterInputs.cpp b/Engine/Plugins/Experimental/ChaosMover/Source/ChaosMover/Private/Character/ChaosCharacterInputs.cpp
--- a/Engine/Plugins/Experimental/ChaosMover/Source/ChaosMover/Private/Character/ChaosCharacterInputs.cpp
+++ b/Engine/Plugins/Experimental/ChaosMover/Source/ChaosMover/Private/Character/ChaosCharacterInputs.cpp
@@ -18,7 +18,7 @@ void FChaosMoverLaunchInputs::ToString(FAnsiStringBuilderBase& Out) const
 	Super::ToString(Out);
 
 	Out.Appendf("LaunchVelocityOrImpulse: X=%.2f Y=%.2f Z=%.2f |", LaunchVelocityOrImpulse.X, LaunchVelocityOrImpulse.Y, LaunchVelocityOrImpulse.Z);
-	Out.Appendf("Mode: %u/n", Mode);
+	Out.Appendf("Mode: %u\n", Mode);
 }
 
 bool FChaosMoverLaunchInputs::ShouldReconcile(const FMoverDataStructBase& AuthorityState) const
diff --git a/Engine/Plugins/Experimental/ChaosMover/Source/ChaosMover/Private/Character/ChaosCharacterMoverComponent.cpp b/Engine/Plugins/Experimental/ChaosMover/Source/ChaosMover/Private/Character/ChaosCharacterMoverComponent.cpp
--- a/Engine/Plugins/Experimental/ChaosMover/Source/ChaosMover/Private/Character/ChaosCharacterMoverComponent.cpp
+++ b/Engine/Plugins/Experimental/ChaosMover/Source/ChaosMover/Private/Character/ChaosCharacterMoverComponent.cpp
@@ -35,6 +35,27 @@ void UChaosCharacterMoverComponent::ProcessSimulationEvent(const FMoverSimulatio
 	}
 }
 
+void UChaosCharacterMoverComponent::DoQueueNextMode(FName DesiredModeName, bool bShouldReenter)
+{
+	Super::DoQueueNextMode(DesiredModeName, bShouldReenter);
+
+	if (QueuedModeTransitionName != NAME_None)
+	{ 
+		const FName NextModeName = QueuedModeTransitionName;
+		if (NextModeName != NAME_None)
+		{
+			UE_LOG(LogMover, Log, TEXT("%s (%s) Overwriting of queued mode change (%s) with (%s)"), *GetNameSafe(GetOwner()), *UEnum::GetValueAsString(GetOwner()->GetLocalRole()), *NextModeName.ToString(), *DesiredModeName.ToString());
+		}
+	}
+
+	QueuedModeTransitionName = DesiredModeName;
+}
+
+void UChaosCharacterMoverComponent::ClearQueuedMode()
+{
+	QueuedModeTransitionName = NAME_None;
+}
+
 void UChaosCharacterMoverComponent::SetAdditionalSimulationOutput(const FMoverDataCollection& Data)
 {
 	Super::SetAdditionalSimulationOutput(Data);
@@ -63,17 +84,25 @@ void UChaosCharacterMoverComponent::ProduceInput(const int32 DeltaTimeMS, FMover
 {
 	Super::ProduceInput(DeltaTimeMS, Cmd);
 
+	check(Cmd);
+
 	// Add launch input data if launch velocity is set
 	if (!LaunchVelocityOrImpulse.IsZero())
 	{
-		check(Cmd);
-
 		FChaosMoverLaunchInputs& LaunchInputs = Cmd->InputCollection.FindOrAddMutableDataByType<FChaosMoverLaunchInputs>();
 		LaunchInputs.LaunchVelocityOrImpulse = LaunchVelocityOrImpulse;
 		LaunchInputs.Mode = LaunchMode;
 
 		LaunchVelocityOrImpulse = FVector::ZeroVector;
 	}
+
+	if (QueuedModeTransitionName != NAME_None)
+	{
+		FCharacterDefaultInputs& CharacterDefaultInputs = Cmd->InputCollection.FindOrAddMutableDataByType<FCharacterDefaultInputs>();
+		CharacterDefaultInputs.SuggestedMovementMode = QueuedModeTransitionName;
+
+		ClearQueuedMode();
+	}
 }
 
 void UChaosCharacterMoverComponent::Launch(const FVector& VelocityOrImpulse, EChaosMoverVelocityEffectMode Mode)
diff --git a/Engine/Plugins/Experimental/ChaosMover/Source/ChaosMover/Private/PathedMovement/ChaosPathedMovementControllerComponent.cpp b/Engine/Plugins/Experimental/ChaosMover/Source/ChaosMover/Private/PathedMovement/ChaosPathedMovementControllerComponent.cpp
new file mode 100644
--- /dev/null
+++ b/Engine/Plugins/Experimental/ChaosMover/Source/ChaosMover/Private/PathedMovement/ChaosPathedMovementControllerComponent.cpp
@@ -0,0 +1,317 @@
+// Copyright Epic Games, Inc. All Rights Reserved.
+
+#include "ChaosMover/PathedMovement/ChaosPathedMovementControllerComponent.h"
+#include "ChaosMover/ChaosMoverLog.h"
+#include "ChaosMover/ChaosMoverSettings.h"
+#include "PBDRigidsSolver.h"
+#include "Physics/Experimental/PhysScene_Chaos.h"
+#include "MoverComponent.h"
+#include "TimerManager.h"
+
+#include UE_INLINE_GENERATED_CPP_BY_NAME(ChaosPathedMovementControllerComponent)
+
+UChaosPathedMovementControllerComponent::UChaosPathedMovementControllerComponent()
+{
+	
+}
+
+void UChaosPathedMovementControllerComponent::BeginPlay()
+{
+	Super::BeginPlay();
+
+	// Subscribe to OnPostMovement to process simulation output
+	if (AActor* Owner = GetOwner())
+	{
+		if (UMoverComponent* OwnerMoverComponent = Owner->FindComponentByClass<UMoverComponent>())
+		{
+			MoverComponent = OwnerMoverComponent;
+			MoverComponent->OnPostFinalize.AddDynamic(this, &UChaosPathedMovementControllerComponent::OnPostFinalize);
+			MoverComponent->OnPostSimEventReceived.AddUObject(this, &UChaosPathedMovementControllerComponent::OnPostSimEventReceived);
+		}
+	}
+
+	const UWorld* World = GetWorld();
+	if (GetOwnerRole() == ROLE_Authority || (World && World->IsNetMode(ENetMode::NM_Standalone)))
+	{
+		ExecuteOrScheduleChange(InitialProperties, /*bIsScheduled =*/ true);
+	}
+}
+
+void UChaosPathedMovementControllerComponent::ProduceInput_Implementation(int32 SimTimeMs, FMoverInputCmdContext& InputCmdResult)
+{
+	FChaosPathedMovementInputs& PathedMovementInputs = InputCmdResult.InputCollection.FindOrAddMutableDataByType<FChaosPathedMovementInputs>();
+	PathedMovementInputs.LastChangeFrame = LastChangeFrame;
+	PathedMovementInputs.MovementStartFrame = MovementStartFrame;
+	PathedMovementInputs.Props = PathedMovementMutableProperties;
+}
+
+void UChaosPathedMovementControllerComponent::OnPostFinalize(const FMoverSyncState& SyncState, const FMoverAuxStateContext& AuxState)
+{
+	if (const FChaosPathedMovementState* FoundPostSimState = SyncState.SyncStateCollection.FindDataByType<FChaosPathedMovementState>())
+	{
+		PostSimState = *FoundPostSimState;
+	}
+}
+
+void UChaosPathedMovementControllerComponent::OnPostSimEventReceived(const FMoverSimulationEventData& EventData)
+{
+	if (const FPathedMovementStartedEventData* PathedMovementStartedEventData = EventData.CastTo<FPathedMovementStartedEventData>())
+	{
+		OnPathedMovementStarted.Broadcast();
+	}
+	else if (const FPathedMovementStoppedEventData* PathedMovementStoppedEventData = EventData.CastTo<FPathedMovementStoppedEventData>())
+	{
+		OnPathedMovementStopped.Broadcast(PathedMovementStoppedEventData->bReachedEndOfPlayback);
+	}
+	else if (const FPathedMovementBouncedEventData* PathedMovementBouncedEventData = EventData.CastTo<FPathedMovementBouncedEventData>())
+	{
+		OnPathedMovementBounced.Broadcast();
+	}
+	else if (const FPathedMovementInReverseChangedEventData* PathedMovementInReverseChangedEventData = EventData.CastTo<FPathedMovementInReverseChangedEventData>())
+	{
+		OnPathedMovementReversePlaybackChanged.Broadcast(PathedMovementInReverseChangedEventData->bNewInReverse);
+	}
+	else if (const FPathedMovementLoopingChangedEventData* PathedMovementLoopingChangedEventData = EventData.CastTo<FPathedMovementLoopingChangedEventData>())
+	{
+		OnPathedMovementLoopingPlaybackChanged.Broadcast(PathedMovementLoopingChangedEventData->bNewLooping);
+	}
+	else if (const FPathedMovementOneWayChangedEventData* PathedMovementOneWayChangedEventData = EventData.CastTo<FPathedMovementOneWayChangedEventData>())
+	{
+		OnPathedMovementOneWayPlaybackChanged.Broadcast(PathedMovementOneWayChangedEventData->bNewOneWay);
+	}
+}
+
+bool ShouldExecute(EChaosPathedMovementExecutionType ExecutionType, AActor* OwnerActor, ENetMode NetMode, const TCHAR* FunctionNameForDebugging)
+{
+	bool bShouldExecute = false;
+
+	if (ensure(OwnerActor))
+	{
+		// Always execute in standalone
+		if (NetMode == NM_Standalone)
+		{
+			bShouldExecute = true;
+		}
+		else
+		{
+			ENetRole OwnerNetRole = OwnerActor->GetLocalRole();
+			if (ExecutionType == EChaosPathedMovementExecutionType::AuthorityOnly)
+			{
+				// Only execute if we're authority on the actor
+				bShouldExecute = (OwnerNetRole == ROLE_Authority);
+			}
+			else
+			{
+				// Don't execute on the server, unless it is locally controlled
+				if (NetMode == NM_DedicatedServer)
+				{
+					bShouldExecute = false;
+				}
+				else if (NetMode == NM_ListenServer)
+				{
+					// Predict on the listen server if it is the locally controlled pawn or controller
+					ENetRole RemoteRole = OwnerActor->GetRemoteRole();
+					bShouldExecute = (RemoteRole == ROLE_AutonomousProxy);
+				}
+				else
+				{
+					bShouldExecute = (OwnerNetRole != ROLE_SimulatedProxy);
+				}
+			}
+		}
+
+		if (!bShouldExecute)
+		{
+			UE_LOG(LogChaosMover, Log, TEXT("%s will be ignored, called with ExecutionType = %s, OwnerLocalRole = %s, OwnerRemoteRole = %s, NetMode = %s."),
+				FunctionNameForDebugging,
+				*StaticEnum<EChaosPathedMovementExecutionType>()->GetNameStringByValue((int64)ExecutionType),
+				*StaticEnum<ENetRole>()->GetNameStringByValue((int64)OwnerActor->GetLocalRole()),
+				*StaticEnum<ENetRole>()->GetNameStringByValue((int64)OwnerActor->GetRemoteRole()),
+				*ToString(NetMode));
+		}
+	}
+	return bShouldExecute;
+}
+
+void UChaosPathedMovementControllerComponent::RequestStartPlayingPath(EChaosPathedMovementExecutionType ExecutionType, bool bIsScheduled)
+{
+	if (!ShouldExecute(ExecutionType, GetOwner(), GetNetMode(), TEXT("UChaosPathedMovementControllerComponent::RequestStartPlayingPath")))
+	{
+		return;
+	}
+
+	Chaos::FPhysicsSolver* Solver = GetPhysicsSolver();
+	const UWorld* World = GetWorld();
+	if (!ensure(Solver) || !World)
+	{
+		return;
+	}
+	
+	if (!PostSimState.PropertiesInEffect.bWantsToPlay || PostSimState.bHasFinished)
+	{
+		FChaosMutablePathedMovementProperties RequestedPathedMovementMutableProperties = PathedMovementMutableProperties;
+		RequestedPathedMovementMutableProperties.bWantsToPlay = true;
+		ExecuteOrScheduleChange(RequestedPathedMovementMutableProperties, bIsScheduled, /*bIsStart =*/ true);
+	}
+}
+
+void UChaosPathedMovementControllerComponent::RequestStopPlayingPath(EChaosPathedMovementExecutionType ExecutionType, bool bIsScheduled)
+{
+	if (!ShouldExecute(ExecutionType, GetOwner(), GetNetMode(), TEXT("UChaosPathedMovementControllerComponent::RequestStopPlayingPath")))
+	{
+		return;
+	}
+
+	FChaosMutablePathedMovementProperties RequestedPathedMovementMutableProperties = PathedMovementMutableProperties;
+	RequestedPathedMovementMutableProperties.bWantsToPlay = false;
+	ExecuteOrScheduleChange(RequestedPathedMovementMutableProperties, bIsScheduled);
+}
+
+bool UChaosPathedMovementControllerComponent::WantsPlayingPath() const
+{
+	return PathedMovementMutableProperties.bWantsToPlay;
+}
+
+bool UChaosPathedMovementControllerComponent::IsPlayingPath() const
+{
+	return PostSimState.PropertiesInEffect.bWantsToPlay && !PostSimState.bHasFinished;
+}
+
+void UChaosPathedMovementControllerComponent::RequestReversePlayback(bool bWantsReversePlayback, EChaosPathedMovementExecutionType ExecutionType, bool bIsScheduled)
+{
+	if (!ShouldExecute(ExecutionType, GetOwner(), GetNetMode(), TEXT("UChaosPathedMovementControllerComponent::RequestReversePlayback")))
+	{
+		return;
+	}
+
+	FChaosMutablePathedMovementProperties RequestedPathedMovementMutableProperties = PathedMovementMutableProperties;
+	RequestedPathedMovementMutableProperties.bWantsReversePlayback = bWantsReversePlayback;
+	ExecuteOrScheduleChange(RequestedPathedMovementMutableProperties, bIsScheduled);
+}
+
+bool UChaosPathedMovementControllerComponent::WantsReversePlayback() const
+{
+	return PathedMovementMutableProperties.bWantsReversePlayback;
+}
+
+bool UChaosPathedMovementControllerComponent::IsReversePlayback() const
+{
+	return PostSimState.PropertiesInEffect.bWantsReversePlayback;
+}
+
+void UChaosPathedMovementControllerComponent::RequestLoopingPlayback(bool bWantsLoopingPlayback, EChaosPathedMovementExecutionType ExecutionType, bool bIsScheduled)
+{
+	if (!ShouldExecute(ExecutionType, GetOwner(), GetNetMode(), TEXT("UChaosPathedMovementControllerComponent::RequestLoopingPlayback")))
+	{
+		return;
+	}
+
+	FChaosMutablePathedMovementProperties RequestedPathedMovementMutableProperties = PathedMovementMutableProperties;
+	RequestedPathedMovementMutableProperties.bWantsLoopingPlayback = bWantsLoopingPlayback;
+	ExecuteOrScheduleChange(RequestedPathedMovementMutableProperties, bIsScheduled);
+}
+
+bool UChaosPathedMovementControllerComponent::WantsLoopingPlayback() const
+{
+	return PathedMovementMutableProperties.bWantsLoopingPlayback;
+}
+
+bool UChaosPathedMovementControllerComponent::IsLoopingPlayback() const
+{
+	return PostSimState.PropertiesInEffect.bWantsLoopingPlayback;
+}
+
+void UChaosPathedMovementControllerComponent::RequestOneWayPlayback(bool bWantsOneWayPlayback, EChaosPathedMovementExecutionType ExecutionType, bool bIsScheduled)
+{
+	if (!ShouldExecute(ExecutionType, GetOwner(), GetNetMode(), TEXT("UChaosPathedMovementControllerComponent::RequestOneWayPlayback")))
+	{
+		return;
+	}
+
+	FChaosMutablePathedMovementProperties RequestedPathedMovementMutableProperties = PathedMovementMutableProperties;
+	RequestedPathedMovementMutableProperties.bWantsOneWayPlayback = bWantsOneWayPlayback;
+	ExecuteOrScheduleChange(RequestedPathedMovementMutableProperties, bIsScheduled);
+}
+
+bool UChaosPathedMovementControllerComponent::WantsOneWayPlayback() const
+{
+	return PathedMovementMutableProperties.bWantsOneWayPlayback;
+}
+
+bool UChaosPathedMovementControllerComponent::IsOneWayPlayback() const
+{
+	return PostSimState.PropertiesInEffect.bWantsOneWayPlayback;
+}
+
+Chaos::FPhysicsSolver* UChaosPathedMovementControllerComponent::GetPhysicsSolver() const
+{
+	if (UWorld* World = GetWorld())
+	{
+		if (FPhysScene* Scene = World->GetPhysicsScene())
+		{
+			return Scene->GetSolver();
+		}
+	}
+	return nullptr;
+}
+
+int32 UChaosPathedMovementControllerComponent::GetEarliestFrameForChange() const
+{
+	if (Chaos::FPhysicsSolver* Solver = GetPhysicsSolver(); ensure(Solver))
+	{
+		// The delay needs to be large enough to cover the largest round trip time but for a consistent
+		// gameplay experience it is better this delay be a constant. If we want to ensure the same gameplay
+		// between a standalone and server game it should be maintained when in standalone as well
+		const float AsyncDeltaTime = Solver->GetAsyncDeltaTime();
+		const float MinDelaySeconds = GetDefault<UChaosMoverSettings>()->PathedMovementSchedulingDelaySeconds;
+		const int32 DelayFrames = FMath::CeilToInt32(MinDelaySeconds / AsyncDeltaTime);
+		return Solver->GetCurrentFrame() + DelayFrames;
+	}
+
+	return 0;
+}
+
+void UChaosPathedMovementControllerComponent::ExecuteOrScheduleChange(const FChaosMutablePathedMovementProperties& NewProperties, bool bIsScheduled, bool bIsStart)
+{
+	Chaos::FPhysicsSolver* Solver = GetPhysicsSolver();
+	if (!Solver)
+	{
+		return;
+	}
+
+	auto LambdaChangeMutableProperties = [&](const FChaosMutablePathedMovementProperties& InProperties, int32 ChangeFrame, bool bIsStart)
+	{
+		LastChangeFrame = ChangeFrame;
+		FChaosMutablePathedMovementProperties PreviousProperties = PathedMovementMutableProperties;
+		PathedMovementMutableProperties = InProperties;
+		if (NewProperties.bWantsToPlay && bIsStart)
+		{
+			// Note that we do this even if we were already playing
+			// This is to be able to restart a non looping path movement which had reached the end
+			// without having to first stop it then start it again
+			MovementStartFrame = ChangeFrame;
+		}
+	};
+
+	if (Solver->IsUsingAsyncResults())
+	{ 
+		LambdaChangeMutableProperties(NewProperties, bIsScheduled ? GetEarliestFrameForChange() : Solver->GetCurrentFrame(), bIsStart);
+	}
+	else if (const UWorld* World = GetWorld())
+	{
+		// In a standalone game that isn't using async physics, the start delay can't be reliably converted to a physics frame.
+		// So we have to use a world timer instead of delaying the start frame
+
+		FTimerManager& TimerManager = World->GetTimerManager();
+		if (PathedMovementDelayedStartTimerHandle.IsValid())
+		{
+			TimerManager.ClearTimer(PathedMovementDelayedStartTimerHandle);
+		}
+
+		const float MinDelaySeconds = GetDefault<UChaosMoverSettings>()->PathedMovementSchedulingDelaySeconds;
+		TimerManager.SetTimer(PathedMovementDelayedStartTimerHandle, FTimerDelegate::CreateWeakLambda(this, [&]()
+		{
+			LambdaChangeMutableProperties(NewProperties, Solver->GetCurrentFrame(), bIsStart);
+		}), MinDelaySeconds, false);
+	}
+}
\ No newline at end of file
diff --git a/Engine/Plugins/Experimental/ChaosMover/Source/ChaosMover/Private/PathedMovement/ChaosPathedMovementMode.cpp b/Engine/Plugins/Experimental/ChaosMover/Source/ChaosMover/Private/PathedMovement/ChaosPathedMovementMode.cpp
new file mode 100644
--- /dev/null
+++ b/Engine/Plugins/Experimental/ChaosMover/Source/ChaosMover/Private/PathedMovement/ChaosPathedMovementMode.cpp
@@ -0,0 +1,553 @@
+// Copyright Epic Games, Inc. All Rights Reserved.
+
+#include "ChaosMover/PathedMovement/ChaosPathedMovementMode.h"
+
+#include "ChaosMover/PathedMovement/ChaosPathedMovementPatternBase.h"
+#include "ChaosMover/ChaosMoverSimulation.h"
+#include "ChaosMover/ChaosMoverLog.h"
+#include "MoverComponent.h"
+#include "Physics/ImmediatePhysics/ImmediatePhysicsAdapters.h"
+#include "Framework/Threading.h"
+#include "Kismet/KismetMathLibrary.h"
+#include "Misc/DataValidation.h"
+
+#if WITH_EDITOR
+#include "Engine/SimpleConstructionScript.h"
+#include "Engine/SCS_Node.h"
+#include "Components/PrimitiveComponent.h"
+#endif
+
+#include UE_INLINE_GENERATED_CPP_BY_NAME(ChaosPathedMovementMode)
+
+#define LOCTEXT_NAMESPACE "UChaosPathedMovementMode"
+
+UChaosPathedMovementMode::UChaosPathedMovementMode(const FObjectInitializer& ObjectInitializer) : Super(ObjectInitializer)
+{
+	bSupportsAsync = true;
+
+	// Sets stiff constraint properties by default
+	SetDefaultJointConstraintProperties();
+}
+
+void UChaosPathedMovementMode::OnRegistered(const FName ModeName)
+{
+	Chaos::EnsureIsInGameThreadContext();
+
+	Super::OnRegistered(ModeName);
+
+	InitializePath();
+
+	ImmediatePhysics::UpdateJointSettingsFromConstraintProfile(JointConstraintProperties, JointSettings);
+}
+
+bool UChaosPathedMovementMode::ShouldUseConstraint() const
+{
+	return bUseJointConstraint;
+}
+
+const Chaos::FPBDJointSettings& UChaosPathedMovementMode::GetConstraintSettings() const
+{
+	return JointSettings;
+}
+
+void UChaosPathedMovementMode::GenerateMove_Implementation(const FMoverTickStartData& StartState, const FMoverTimeStep& TimeStep, FProposedMove& OutProposedMove) const
+{
+}
+
+void UChaosPathedMovementMode::SimulationTick_Implementation(const FSimulationTickParams& Params, FMoverTickEndData& OutputState)
+{
+	Chaos::EnsureIsInPhysicsThreadContext();
+
+#if !UE_BUILD_SHIPPING
+	DebugData.Advance();
+#endif // !UE_BUILD_SHIPPING
+
+	// We retrieve the previous properties so we can compare with the new. If there are none we default to FChaosPathedMovementState()
+	const FChaosPathedMovementState* FoundPreviousMoveState = Params.StartState.SyncState.SyncStateCollection.FindDataByType<FChaosPathedMovementState>();
+	const FChaosPathedMovementState PreviousMoveState = FoundPreviousMoveState ? *FoundPreviousMoveState : FChaosPathedMovementState();	
+	const FChaosMutablePathedMovementProperties& PreviousProperties = PreviousMoveState.PropertiesInEffect;
+
+	const FChaosPathedMovementInputs* PathedMovementInputs = Params.StartState.InputCmd.InputCollection.FindDataByType<FChaosPathedMovementInputs>();
+	const FChaosMutablePathedMovementProperties& InputProperties = PathedMovementInputs ? PathedMovementInputs->Props : PreviousProperties;
+
+	// We initialize the output move state from the previous move state
+	FChaosPathedMovementState& OutputPathedMovemenState = OutputState.SyncState.SyncStateCollection.FindOrAddMutableDataByType<FChaosPathedMovementState>();
+	// We initialize the output movement state from the previous movement state
+	OutputPathedMovemenState = PreviousMoveState;
+
+	
+	bool bReceivedNewStartRequest = false;
+
+	// If we have inputs,
+	if (PathedMovementInputs
+		// ... and we have reached the frame at which they should enter into effect
+		&& (Params.TimeStep.ServerFrame >= PathedMovementInputs->LastChangeFrame)
+		// ... and they are different than the ones in effect currently
+		//         or the input's MovementStartFrame is later than the output's movement start frame
+		&& ((PreviousProperties != PathedMovementInputs->Props)
+			|| (PathedMovementInputs->MovementStartFrame > OutputPathedMovemenState.LatestMovementStartFrame)))
+	{
+		bReceivedNewStartRequest = PathedMovementInputs->Props.bWantsToPlay
+			&& (!PreviousProperties.bWantsToPlay
+			    || PathedMovementInputs->MovementStartFrame > OutputPathedMovemenState.LatestMovementStartFrame);
+
+		// ... then the input properties enter into effect
+		OutputPathedMovemenState.PropertiesInEffect = PathedMovementInputs->Props;
+	}
+
+	// PropertiesInEffect is const here, to indicate it is a bad idea to modify those, since they are used to detect external changes in mutable properties,
+	// which would become confusing if modified internally.
+	const FChaosMutablePathedMovementProperties& CurrentProperties = OutputPathedMovemenState.PropertiesInEffect;
+
+	// These booleans will determine whether events get fired at the end of this 
+	bool bHasJustFinished = false;
+	bool bHasJustBounced = false;
+	bool bJustRestarted = false;
+	bool bWasStoppedJustRequested = PreviousProperties.bWantsToPlay && !CurrentProperties.bWantsToPlay;
+
+	// When bWantsReversePlayback changes and we are still playing, we change the direction of motion
+	bool bIsInReverseJustChanged = PreviousProperties.bWantsReversePlayback != CurrentProperties.bWantsReversePlayback;
+	if (!OutputPathedMovemenState.bHasFinished && bIsInReverseJustChanged)
+	{
+		OutputPathedMovemenState.bIsPathFractionIncreasing = !OutputPathedMovemenState.bIsPathFractionIncreasing;
+	}
+
+	// Reset bHasFinished if we are restarting
+	if (bReceivedNewStartRequest && OutputPathedMovemenState.bHasFinished)
+	{
+		OutputPathedMovemenState.bHasFinished = false;
+		// We initialize the direction of motion to respect reverse playback on start
+		OutputPathedMovemenState.bIsPathFractionIncreasing = !CurrentProperties.bWantsReversePlayback;
+		if (CurrentProperties.bWantsOneWayPlayback)
+		{
+			if (OutputPathedMovemenState.CurrentPathFraction == (OutputPathedMovemenState.bIsPathFractionIncreasing ? 1.0f : 0.0f))
+			{
+				OutputPathedMovemenState.CurrentPathFraction = 1.0f - OutputPathedMovemenState.CurrentPathFraction;
+			}
+		}
+	}
+
+	bool bLoopingJustChanged = PreviousProperties.bWantsLoopingPlayback != CurrentProperties.bWantsLoopingPlayback;
+	if (bLoopingJustChanged && CurrentProperties.bWantsToPlay && OutputPathedMovemenState.bHasFinished)
+	{
+		OutputPathedMovemenState.bHasFinished = false;
+		bJustRestarted = true;
+	}
+
+	if (CurrentProperties.bWantsToPlay && (!PreviousMoveState.bHasFinished || bReceivedNewStartRequest))
+	{
+		const float FrameDuration = Params.TimeStep.StepMs * 0.001f;
+		const float PlaybackSpeed = (OneWayPlaybackDuration >= 0.0f) ? (1.0f / OneWayPlaybackDuration) : 0.0f;
+		const float Direction = OutputPathedMovemenState.bIsPathFractionIncreasing ? 1.0f : -1.0f;
+		const float DeltaPathFraction = Direction * PlaybackSpeed * FrameDuration;
+		const float NewPathFraction = OutputPathedMovemenState.CurrentPathFraction + DeltaPathFraction;
+
+		if (CurrentProperties.bWantsOneWayPlayback)
+		{
+			// One way forward playback
+			if (OutputPathedMovemenState.bIsPathFractionIncreasing)
+			{
+				if (NewPathFraction >= 1.0f)
+				{
+					if (CurrentProperties.bWantsLoopingPlayback)
+					{
+						OutputPathedMovemenState.CurrentPathFraction = FMath::Fmod(NewPathFraction, 1.0f);
+						bJustRestarted = true;
+					}
+					else
+					{
+						OutputPathedMovemenState.CurrentPathFraction = 1.0f;
+						OutputPathedMovemenState.bHasFinished = true;
+						bHasJustFinished = true;
+					}
+				}
+				else
+				{
+					OutputPathedMovemenState.CurrentPathFraction = NewPathFraction;
+				}
+			}
+			else // One way reverse playback
+			{
+				if (NewPathFraction <= 0.0f)
+				{
+					if (CurrentProperties.bWantsLoopingPlayback)
+					{
+						OutputPathedMovemenState.CurrentPathFraction = 1.0f - FMath::Fmod(-NewPathFraction, 1.0f);
+						bJustRestarted = true;
+					}
+					else
+					{
+						OutputPathedMovemenState.CurrentPathFraction = 0.0f;
+						OutputPathedMovemenState.bHasFinished = true;
+						bHasJustFinished = true;
+					}
+				}
+				else
+				{
+					OutputPathedMovemenState.CurrentPathFraction = NewPathFraction;
+				}
+			}
+		}
+		else
+		{
+			// Round-trip playback
+			if (OutputPathedMovemenState.bIsPathFractionIncreasing && NewPathFraction >= 1.0f)
+			{
+				if (!CurrentProperties.bWantsReversePlayback || CurrentProperties.bWantsLoopingPlayback)
+				{
+					OutputPathedMovemenState.bIsPathFractionIncreasing = false;
+					OutputPathedMovemenState.CurrentPathFraction = 1.0f - FMath::Fmod(NewPathFraction, 1.0f);
+					bHasJustBounced = true;
+				}
+				else
+				{
+					OutputPathedMovemenState.CurrentPathFraction = 1.0f;
+					OutputPathedMovemenState.bHasFinished = true;
+					bHasJustFinished = true;
+				}
+			}
+			else if (!OutputPathedMovemenState.bIsPathFractionIncreasing && NewPathFraction <= 0.0f)
+			{
+				if (CurrentProperties.bWantsReversePlayback || CurrentProperties.bWantsLoopingPlayback)
+				{
+					OutputPathedMovemenState.bIsPathFractionIncreasing = true;
+					OutputPathedMovemenState.CurrentPathFraction = FMath::Fmod(-NewPathFraction, 1.0f);
+					bHasJustBounced = true;
+				}
+				else
+				{
+					OutputPathedMovemenState.CurrentPathFraction = 0.0f;
+					OutputPathedMovemenState.bHasFinished = true;
+					bHasJustFinished = true;
+				}
+			}
+			else
+			{
+				OutputPathedMovemenState.CurrentPathFraction = NewPathFraction;
+			}
+		}
+	}
+
+	if (bIsInReverseJustChanged)
+	{
+		OnInReverseChanging(Params.TimeStep.BaseSimTimeMs, CurrentProperties.bWantsReversePlayback);
+	}
+	if (bLoopingJustChanged)
+	{
+		OnLoopingChanging(Params.TimeStep.BaseSimTimeMs, CurrentProperties.bWantsLoopingPlayback);
+	}
+	bool bOneWayJustChanged = CurrentProperties.bWantsOneWayPlayback != PreviousProperties.bWantsOneWayPlayback;
+	if (bOneWayJustChanged)
+	{
+		OnOneWayChanging(Params.TimeStep.BaseSimTimeMs, CurrentProperties.bWantsLoopingPlayback);
+	}
+
+	if (bReceivedNewStartRequest)
+	{
+		OnPlaybackStarting(Params.TimeStep.BaseSimTimeMs);
+		OutputPathedMovemenState.LatestMovementStartFrame = Params.TimeStep.ServerFrame;
+	}
+
+	if ((bWasStoppedJustRequested && !OutputPathedMovemenState.bHasFinished) || bHasJustFinished)
+	{
+		OnPlaybackStopping(Params.TimeStep.BaseSimTimeMs, bHasJustFinished);
+	}
+
+	if (bHasJustBounced)
+	{
+		OnPlaybackBounced(Params.TimeStep.BaseSimTimeMs);
+	}
+
+	const Chaos::FRigidTransform3 TargetRelativeTransform = CalcTargetRelativeTransform(OutputPathedMovemenState.CurrentPathFraction);
+
+	FMoverDefaultSyncState& OutputDefaultSyncState = OutputState.SyncState.SyncStateCollection.FindOrAddMutableDataByType<FMoverDefaultSyncState>();
+	
+	const FTransform& PathBasis = Simulation->GetPathedMovementBasis();
+	FVector WorldTranslation(PathBasis.TransformPosition(TargetRelativeTransform.GetTranslation()));
+	FRotator WorldRotation(PathBasis.TransformRotation(TargetRelativeTransform.GetRotation()));
+	OutputDefaultSyncState.SetTransforms_WorldSpace( WorldTranslation, WorldRotation,
+											  FVector::ZeroVector, // Fix this eventually to return the velocity tangential to the path, including due to a change in path basis since last update
+											  nullptr ); // no movement base
+
+#if !UE_BUILD_SHIPPING
+	DebugData.ServerFrame = Params.TimeStep.ServerFrame;
+	DebugData.PreSimMoveState = PreviousMoveState;
+	DebugData.PostSimMoveState = OutputPathedMovemenState;
+	DebugData.FinalWorldTranslation = WorldTranslation;
+	DebugData.FinalWorldRotation = WorldRotation;
+	// Record state for debugging
+	UE::ChaosMover::GetDebugSimData(Simulation).FindOrAddMutableDataByType<FChaosPathedMovementModeDebugData>() = DebugData;
+#endif
+}
+
+void UChaosPathedMovementMode::OnPlaybackStarting(float BaseSimTimeMs)
+{
+	// Add a pathed movement started event to the simulation (will be broadcast on GT during post sim)
+	if (Simulation)
+	{
+		Simulation->AddEvent(MakeShared<FPathedMovementStartedEventData>(BaseSimTimeMs));
+
+#if !UE_BUILD_SHIPPING
+		DebugData.bPlaybackStartedThisFrame = true;
+#endif // !UE_BUILD_SHIPPING
+	}
+}
+
+void UChaosPathedMovementMode::OnPlaybackStopping(float BaseSimTimeMs, bool bReachedEndOfPlayback)
+{
+	// Add a pathed movement stopped event to the simulation (will be broadcast on GT during post sim)
+	if (Simulation)
+	{
+		Simulation->AddEvent(MakeShared<FPathedMovementStoppedEventData>(BaseSimTimeMs, bReachedEndOfPlayback));
+
+#if !UE_BUILD_SHIPPING
+		DebugData.bPlaybackStoppedThisFrame = true;
+		DebugData.bPlaybackReachedEndThisFrame = bReachedEndOfPlayback;
+#endif // !UE_BUILD_SHIPPING
+	}
+}
+
+void UChaosPathedMovementMode::OnPlaybackBounced(float BaseSimTimeMs)
+{
+	// Add a pathed movement bounced event to the simulation (will be broadcast on GT during post sim)
+	if (Simulation)
+	{
+		Simulation->AddEvent(MakeShared<FPathedMovementBouncedEventData>(BaseSimTimeMs));
+
+#if !UE_BUILD_SHIPPING
+		DebugData.bPlaybackBouncedThisFrame = true;
+#endif // !UE_BUILD_SHIPPING
+	}
+}
+
+void UChaosPathedMovementMode::OnInReverseChanging(float BaseSimTimeMs, bool bNewInReverse)
+{
+	// Add a pathed movement stopped event to the simulation (will be broadcast on GT during post sim)
+	if (Simulation)
+	{
+		Simulation->AddEvent(MakeShared<FPathedMovementInReverseChangedEventData>(BaseSimTimeMs, bNewInReverse));
+
+#if !UE_BUILD_SHIPPING
+		DebugData.bPlaybackReverseChangedThisFrame = true;
+#endif // !UE_BUILD_SHIPPING
+	}
+}
+
+void UChaosPathedMovementMode::OnLoopingChanging(float BaseSimTimeMs, bool bNewLooping)
+{
+	// Add a pathed movement stopped event to the simulation (will be broadcast on GT during post sim)
+	if (Simulation)
+	{
+		Simulation->AddEvent(MakeShared<FPathedMovementLoopingChangedEventData>(BaseSimTimeMs, bNewLooping));
+
+#if !UE_BUILD_SHIPPING
+		DebugData.bPlaybackLoopingChangedThisFrame = true;
+#endif // !UE_BUILD_SHIPPING
+	}
+}
+
+void UChaosPathedMovementMode::OnOneWayChanging(float BaseSimTimeMs, bool bNewOneWay)
+{
+	// Add a pathed movement stopped event to the simulation (will be broadcast on GT during post sim)
+	if (Simulation)
+	{
+		Simulation->AddEvent(MakeShared<FPathedMovementOneWayChangedEventData>(BaseSimTimeMs, bNewOneWay));
+
+#if !UE_BUILD_SHIPPING
+		DebugData.bPlaybackOneWayChangedThisFrame = true;
+#endif // !UE_BUILD_SHIPPING
+	}
+}
+
+void UChaosPathedMovementMode::InitializePath()
+{
+	Chaos::EnsureIsInGameThreadContext();
+
+	for (UChaosPathedMovementPatternBase* Pattern : PathPatterns)
+	{
+		if (Pattern)
+		{
+			Pattern->InitializePattern(Simulation);
+		}
+	}
+}
+
+UChaosPathedMovementPatternBase* UChaosPathedMovementMode::BP_FindPattern(TSubclassOf<UChaosPathedMovementPatternBase> PatternType) const
+{
+	for (UChaosPathedMovementPatternBase* Pattern : PathPatterns)
+	{
+		if (Pattern && Pattern->IsA(PatternType))
+		{
+			return Pattern;
+		}
+	}
+	
+	return nullptr;
+}
+
+void UChaosPathedMovementMode::SetOneWayTripDuration_BeginPlayOnly(float NewDuration)
+{
+	Chaos::EnsureIsInGameThreadContext();
+
+	if (const UMoverComponent* MoverComponent = GetMoverComponent())
+	{
+		if (ensure(!MoverComponent->HasBegunPlay()))
+		{
+			OneWayPlaybackDuration = FMath::Max(NewDuration, 0.f);
+		}
+	}
+}
+
+Chaos::FRigidTransform3 UChaosPathedMovementMode::CalcTargetRelativeTransform(float ProgressAmt) const
+{
+	FTransform TargetRelativeTransform = FTransform::Identity;
+	
+	for (UChaosPathedMovementPatternBase* PathPattern : PathPatterns)
+	{
+		if (PathPattern)
+		{
+			const FTransform PatternTransform = PathPattern->CalcTargetRelativeTransform(ProgressAmt, TargetRelativeTransform);
+        	TargetRelativeTransform.Accumulate(PatternTransform);
+		}
+	}
+
+	return TargetRelativeTransform;
+}
+
+void UChaosPathedMovementMode::SetDefaultJointConstraintProperties()
+{
+	JointConstraintProperties.bLinearBreakable = false;
+	JointConstraintProperties.bAngularBreakable = false;
+	JointConstraintProperties.bDisableCollision = true;
+
+	// Free the linear joints and set the max allowed distance from the target
+	static constexpr float DefaultMaxLinearDistance = 100.f;
+	JointConstraintProperties.LinearLimit.XMotion = LCM_Free;
+	JointConstraintProperties.LinearLimit.YMotion = LCM_Free;
+	JointConstraintProperties.LinearLimit.ZMotion = LCM_Free;
+	JointConstraintProperties.LinearLimit.Limit = DefaultMaxLinearDistance;
+
+	static constexpr float DefaultDriveMaxForce = 5000.f;
+	
+	// Linear drive config
+	static constexpr float LinearDriveStiffness = 750.f;
+	static const float LinearDriveCriticalDamping = 2 * FMath::Sqrt(LinearDriveStiffness);
+	
+	FConstraintDrive DefaultLinearDriveConstraint;
+	DefaultLinearDriveConstraint.bEnablePositionDrive = true;
+	DefaultLinearDriveConstraint.bEnableVelocityDrive = true;
+	DefaultLinearDriveConstraint.Stiffness = LinearDriveStiffness;
+	DefaultLinearDriveConstraint.Damping = LinearDriveCriticalDamping;
+	DefaultLinearDriveConstraint.MaxForce = DefaultDriveMaxForce;
+
+	JointConstraintProperties.LinearDrive.XDrive = DefaultLinearDriveConstraint;
+	JointConstraintProperties.LinearDrive.YDrive = DefaultLinearDriveConstraint;
+	JointConstraintProperties.LinearDrive.ZDrive = DefaultLinearDriveConstraint;
+
+	// Free the angular joints and set degree limits relative to the target
+	static constexpr float DefaultAngularLimitDegrees = 15.f;
+	JointConstraintProperties.ConeLimit.Swing1Motion = ACM_Free;
+	JointConstraintProperties.ConeLimit.Swing1LimitDegrees = DefaultAngularLimitDegrees;
+	JointConstraintProperties.ConeLimit.Swing2Motion = ACM_Free;
+	JointConstraintProperties.ConeLimit.Swing2LimitDegrees = DefaultAngularLimitDegrees;
+	JointConstraintProperties.TwistLimit.TwistMotion = ACM_Free;
+	JointConstraintProperties.TwistLimit.TwistLimitDegrees = DefaultAngularLimitDegrees;
+	
+	// Angular drive config
+	static constexpr float AngularDriveStiffness = 1500.f;
+	static const float AngularDriveCriticalDamping = 2 * FMath::Sqrt(AngularDriveStiffness);
+
+	JointConstraintProperties.AngularDrive.AngularDriveMode = EAngularDriveMode::SLERP;
+	JointConstraintProperties.AngularDrive.SlerpDrive.bEnablePositionDrive = true;
+	JointConstraintProperties.AngularDrive.SlerpDrive.bEnableVelocityDrive = true;
+	JointConstraintProperties.AngularDrive.SlerpDrive.Stiffness = AngularDriveStiffness;
+	JointConstraintProperties.AngularDrive.SlerpDrive.Damping = AngularDriveCriticalDamping;
+	JointConstraintProperties.AngularDrive.SlerpDrive.MaxForce = DefaultDriveMaxForce;
+}
+
+#if WITH_EDITOR
+template <typename TComponentType>
+TComponentType* GetOwnerActorsRootComponent(const UActorComponent* ActorComponent)
+{
+	if (!ActorComponent)
+	{
+		return nullptr;
+	}
+
+	const AActor* OuterActor = ActorComponent ? ActorComponent->GetTypedOuter<AActor>() : nullptr;
+	if (OuterActor)
+	{
+		return Cast<TComponentType>(OuterActor->GetRootComponent());
+	}
+	else if (UBlueprintGeneratedClass* OuterBlueprintGeneratedClass = ActorComponent->GetTypedOuter<UBlueprintGeneratedClass>())
+	{
+		TArray<const UBlueprintGeneratedClass*> BlueprintClasses;
+		UBlueprintGeneratedClass::GetGeneratedClassesHierarchy(OuterBlueprintGeneratedClass, BlueprintClasses);
+		for (const UBlueprintGeneratedClass* BPGClass : BlueprintClasses)
+		{
+			if (BPGClass && BPGClass->SimpleConstructionScript)
+			{
+				const TArray<USCS_Node*>& ActorBlueprintNodes = BPGClass->SimpleConstructionScript->GetAllNodes();
+				for (USCS_Node* Node : ActorBlueprintNodes)
+				{
+					if (Node->IsRootNode() && Node->ComponentClass && Node->ComponentClass->IsChildOf(USceneComponent::StaticClass()))
+					{
+						if (TComponentType* CastRootComponent = Cast<TComponentType>(Node->GetActualComponentTemplate(OuterBlueprintGeneratedClass)))
+						{
+							return CastRootComponent;
+						}
+					}
+				}
+			}
+		}
+	}
+
+	return nullptr;
+}
+
+EDataValidationResult UChaosPathedMovementMode::IsDataValid(FDataValidationContext& Context) const
+{
+	EDataValidationResult Result = Super::IsDataValid(Context);
+
+	// Check for correct physics setup of the actor this mode will be moving:
+	//    - its root needs to be a primitive component
+	//        * (this check should be done in UMoverComponent::IsDataValid also, 
+	//           or we should relax this constraint and allow to pick a non root primitive
+	//           component as the updated component)
+	//    - it needs to be simulating physics
+	//    - if kinematically moved, it needs to have bUpdateKinematicFromSimulation set to true
+	//
+	// Find the CDO of the MoverComponent's Actor.
+	// There might not be one if this path mode is on a free blueprinted mover component asset
+	// that's not on an actor. In that case there's nothing to check.
+	const UMoverComponent* MoverComponent = GetMoverComponent();
+	if (!MoverComponent)
+	{
+		return Result;
+	}
+
+	// Find the root scene component
+	if (const USceneComponent* RootSceneComponent = GetOwnerActorsRootComponent<USceneComponent>(MoverComponent))
+	{
+		// Check that the root scene component is a primitive component
+		const UPrimitiveComponent* RootPrimitiveComponent = Cast<UPrimitiveComponent>(RootSceneComponent);
+
+		if (!RootPrimitiveComponent)
+		{
+			Context.AddError(FText::Format(LOCTEXT("RootNotPrimitiveComponent", "The root component {0} of an actor using a {1} must be a primitive component"), FText::FromString(GetNameSafe(RootSceneComponent)), FText::FromString(GetClass()->GetName())));
+		}
+		else if (!RootPrimitiveComponent->BodyInstance.bSimulatePhysics)
+		{
+			Context.AddWarning(FText::Format(LOCTEXT("SuspiciousSimulatingPhysics", "The root component {0} of an actor using a {1} is not simulating physics, its mover component may not work correctly"), FText::FromString(GetNameSafe(RootSceneComponent)), FText::FromString(GetClass()->GetName())));
+		}
+		else if (!bUseJointConstraint && !RootPrimitiveComponent->BodyInstance.bUpdateKinematicFromSimulation)
+		{
+			Context.AddError(FText::Format(LOCTEXT("InvalidUpdateKinematicFromSimulation", "The root component {0} of an actor using a {1} will be moved kinematically, it needs to have Update Kinematic From Simulation checked in its Physics properties"), FText::FromString(GetNameSafe(RootSceneComponent)), FText::FromString(GetClass()->GetName())));
+		}
+	}
+
+	return Result;
+}
+#endif
+
+#undef LOCTEXT_NAMESPACE
diff --git a/Engine/Plugins/Experimental/ChaosMover/Source/ChaosMover/Private/PathedMovement/ChaosPathedMovementPatternBase.cpp b/Engine/Plugins/Experimental/ChaosMover/Source/ChaosMover/Private/PathedMovement/ChaosPathedMovementPatternBase.cpp
new file mode 100644
--- /dev/null
+++ b/Engine/Plugins/Experimental/ChaosMover/Source/ChaosMover/Private/PathedMovement/ChaosPathedMovementPatternBase.cpp
@@ -0,0 +1,112 @@
+// Copyright Epic Games, Inc. All Rights Reserved.
+
+#include "ChaosMover/PathedMovement/ChaosPathedMovementPatternBase.h"
+
+#include "HAL/IConsoleManager.h"
+#include "ChaosMover/PathedMovement/ChaosPathedMovementMode.h"
+#include "ChaosMover/ChaosMoverSimulation.h"
+
+#include UE_INLINE_GENERATED_CPP_BY_NAME(ChaosPathedMovementPatternBase)
+
+void UChaosPathedMovementPatternBase::InitializePattern(UChaosMoverSimulation* InSimulation)
+{
+	Simulation = InSimulation;
+}
+
+void UChaosPathedMovementPatternBase::ProduceInputs_External(OUT FPhysicsMoverAsyncInput& Input)
+{
+	//@todo DanH: Push something with the total duration and playback behavior
+}
+
+template <typename T>
+T ApplyAxisMask(const T& Unmasked, EChaosPatternAxisMaskFlags Flags, float MaskedValue = 0.f)
+{
+	T MaskedResult = Unmasked;
+	if (EnumHasAnyFlags(Flags, EChaosPatternAxisMaskFlags::X))
+	{
+		MaskedResult.X = MaskedValue;
+	}
+	if (EnumHasAnyFlags(Flags, EChaosPatternAxisMaskFlags::Y))
+	{
+		MaskedResult.Y = MaskedValue;
+	}
+	if (EnumHasAnyFlags(Flags, EChaosPatternAxisMaskFlags::Z))
+	{
+		MaskedResult.Z = MaskedValue;
+	}
+	return MaskedResult;
+}
+
+template <>
+FRotator ApplyAxisMask(const FRotator& Unmasked, EChaosPatternAxisMaskFlags Flags, float MaskedValue)
+{
+	FRotator MaskedResult = Unmasked;
+	if (EnumHasAnyFlags(Flags, EChaosPatternAxisMaskFlags::X))
+	{
+		MaskedResult.Roll = MaskedValue;
+	}
+	if (EnumHasAnyFlags(Flags, EChaosPatternAxisMaskFlags::Y))
+	{
+		MaskedResult.Pitch = MaskedValue;
+	}
+	if (EnumHasAnyFlags(Flags, EChaosPatternAxisMaskFlags::Z))
+	{
+		MaskedResult.Yaw = MaskedValue;
+	}
+	return MaskedResult;
+}
+
+FTransform UChaosPathedMovementPatternBase::CalcTargetRelativeTransform(float OverallPathProgress, const FTransform& CurTargetTransform) const
+{
+	//@todo DanH: When doing debug drawing, when a pattern is shorter and/or loops more, it loses resolution. Can't just always calculate as though OverallPathProgress == PatternProgress
+	//		though because we still depend on the calculated result. Maybe I shouldn't bother with that and just let the aggregate one calc on its own?
+	
+	if (NumLoopsPerPath <= 0 || StartAtPathProgress >= EndAtPathProgress || OverallPathProgress <= StartAtPathProgress)
+	{
+		return FTransform();
+	}
+	
+	// How far into the current loop of this specific pattern are we?
+	const float PathProgressSinceStart = FMath::Min(EndAtPathProgress, OverallPathProgress) - StartAtPathProgress;
+	const float PathProgressPerPatternLoop = (EndAtPathProgress - StartAtPathProgress) / NumLoopsPerPath;
+	float CurLoopPathProgress = FMath::Fmod(PathProgressSinceStart, PathProgressPerPatternLoop);
+	if (CurLoopPathProgress == 0.f && PathProgressSinceStart > 0.f)
+	{
+		// Treat progress that matches the per-loop span exactly as 100%, not 0%
+		CurLoopPathProgress = PathProgressPerPatternLoop;
+	}
+
+	// If each loop is there and back, the actual progress needs to flip after completing half the span
+	if (!IsOneWay)
+	{
+		// ThereAndBack means we actually progress twice as fast as a OneShot
+		CurLoopPathProgress *= 2.f;
+		
+		const float ReversePathProgress = CurLoopPathProgress - PathProgressPerPatternLoop;
+		if (ReversePathProgress > 0.f)
+		{
+			CurLoopPathProgress = PathProgressPerPatternLoop - ReversePathProgress;
+		}
+	}
+
+	// Now convert from overall path progress to progress for this pattern
+	const float PatternProgress = CurLoopPathProgress / PathProgressPerPatternLoop;
+	const FTransform UnmaskedTarget = CalcUnmaskedTargetRelativeTransform(PatternProgress, CurTargetTransform);
+
+	if (bOrientComponentToPath)
+	{
+		//@todo DanH: Rotate so that the forward vector matches the translation vector direction
+	}
+
+	const FVector TargetLocation = ApplyAxisMask(UnmaskedTarget.GetLocation(), (EChaosPatternAxisMaskFlags)TranslationMasks);
+	const FRotator TargetRotation = ApplyAxisMask(UnmaskedTarget.GetRotation().Rotator(), (EChaosPatternAxisMaskFlags)RotationMasks);
+	const FVector TargetScale = ApplyAxisMask(UnmaskedTarget.GetScale3D(), (EChaosPatternAxisMaskFlags)ScaleMasks, 1.f);
+	return FTransform(TargetRotation, TargetLocation, TargetScale);
+}
+
+UChaosPathedMovementMode& UChaosPathedMovementPatternBase::GetMovementMode() const
+{
+	UChaosPathedMovementMode* OuterMode = GetOuterUChaosPathedMovementMode();
+	check(OuterMode);
+	return *OuterMode;
+}
diff --git a/Engine/Plugins/Experimental/ChaosMover/Source/ChaosMover/Private/PathedMovement/ChaosPathedMovementTypes.cpp b/Engine/Plugins/Experimental/ChaosMover/Source/ChaosMover/Private/PathedMovement/ChaosPathedMovementTypes.cpp
new file mode 100644
--- /dev/null
+++ b/Engine/Plugins/Experimental/ChaosMover/Source/ChaosMover/Private/PathedMovement/ChaosPathedMovementTypes.cpp
@@ -0,0 +1,163 @@
+// Copyright Epic Games, Inc. All Rights Reserved.
+
+#include "ChaosMover/PathedMovement/ChaosPathedMovementTypes.h"
+
+#include UE_INLINE_GENERATED_CPP_BY_NAME(ChaosPathedMovementTypes)
+
+void FChaosMutablePathedMovementProperties::NetSerialize(FArchive& Ar)
+{
+	Ar.SerializeBits(&bWantsToPlay, 1);
+	Ar.SerializeBits(&bWantsReversePlayback, 1);
+	Ar.SerializeBits(&bWantsLoopingPlayback, 1);
+	Ar.SerializeBits(&bWantsOneWayPlayback, 1);
+}
+
+void FChaosMutablePathedMovementProperties::ToString(FAnsiStringBuilderBase& Out) const
+{
+	Out.Appendf("Wants to Play Path: %hs | \n", *LexToString(bWantsToPlay));
+	Out.Appendf("Wants Reverse Playback: %hs | \n", *LexToString(bWantsReversePlayback));
+	Out.Appendf("Wants Looping Playback: %hs | \n", *LexToString(bWantsLoopingPlayback));
+	Out.Appendf("Wants One Way Playback: %hs | \n", *LexToString(bWantsOneWayPlayback));
+}
+
+bool FChaosMutablePathedMovementProperties::operator==(const FChaosMutablePathedMovementProperties& Other) const
+{
+	return bWantsToPlay == Other.bWantsToPlay
+		&& bWantsReversePlayback == Other.bWantsReversePlayback
+		&& bWantsLoopingPlayback == Other.bWantsLoopingPlayback
+		&& bWantsOneWayPlayback == Other.bWantsOneWayPlayback
+		;
+}
+
+bool FChaosPathedMovementInputs::NetSerialize(FArchive& Ar, UPackageMap* Map, bool& bOutSuccess)
+{
+	Super::NetSerialize(Ar, Map, bOutSuccess);
+
+	Ar << MovementStartFrame;
+	Ar << LastChangeFrame;
+	Props.NetSerialize(Ar);
+	
+	bOutSuccess = true;
+	return true;
+}
+
+void FChaosPathedMovementInputs::ToString(FAnsiStringBuilderBase& Out) const
+{
+	Props.ToString(Out);
+	Out.Appendf("Movement Start Frame: %d\n", MovementStartFrame);
+	Out.Appendf("Last Change Frame: %d\n", LastChangeFrame);
+}
+
+bool FChaosPathedMovementInputs::ShouldReconcile(const FMoverDataStructBase& AuthorityState) const
+{
+	const FChaosPathedMovementInputs& TypedAuthority = static_cast<const FChaosPathedMovementInputs&>(AuthorityState);
+	return Props != TypedAuthority.Props;
+}
+
+void FChaosPathedMovementInputs::Interpolate(const FMoverDataStructBase& From, const FMoverDataStructBase& To, float Pct)
+{
+	if (Pct < .5f)
+	{
+		*this = static_cast<const FChaosPathedMovementInputs&>(From);
+	}
+	else
+	{
+		*this = static_cast<const FChaosPathedMovementInputs&>(To);
+	}
+}
+
+void FChaosPathedMovementInputs::Merge(const FMoverDataStructBase& From)
+{
+	const FChaosPathedMovementInputs& TypedFrom = static_cast<const FChaosPathedMovementInputs&>(From);
+	//@todo DanH: What, if anything, do I actually want to merge on the mutable props?
+}
+
+bool FChaosPathedMovementState::NetSerialize(FArchive& Ar, UPackageMap* Map, bool& bOutSuccess)
+{
+	Super::NetSerialize(Ar, Map, bOutSuccess);
+
+	static FChaosPathedMovementState Defaults;
+
+	bool HasCurrentPathFraction;
+	bool HasLatestMovementStartFrame;
+
+	if (Ar.IsSaving())
+	{
+		HasCurrentPathFraction = CurrentPathFraction != Defaults.CurrentPathFraction;
+		HasLatestMovementStartFrame = LatestMovementStartFrame != Defaults.LatestMovementStartFrame;
+	}
+	else
+	{
+		// Set to defaults
+		*this = Defaults;
+	}
+
+	Ar.SerializeBits(&HasCurrentPathFraction, 1);
+	Ar.SerializeBits(&HasLatestMovementStartFrame, 1);
+	Ar.SerializeBits(&bIsPathFractionIncreasing, 1);
+	Ar.SerializeBits(&bHasFinished, 1);
+	PropertiesInEffect.NetSerialize(Ar);
+
+	if (HasCurrentPathFraction)
+	{
+		Ar << CurrentPathFraction;
+	}
+	if (HasLatestMovementStartFrame)
+	{
+		Ar << LatestMovementStartFrame;
+	}
+
+	return true;
+}
+
+void FChaosPathedMovementState::ToString(FAnsiStringBuilderBase& Out) const
+{
+	FMoverDataStructBase::ToString(Out);
+
+	Out.Appendf("Current Path Fraction: %.2f | ", CurrentPathFraction);
+	Out.Appendf("Latest Movement Start Frame: %d | ", LatestMovementStartFrame);
+	Out.Appendf("Is Path Fraction Increasing: %s | ", bIsPathFractionIncreasing ? TEXT("True") : TEXT("False"));
+	Out.Appendf("Has Finished: %s | ", bHasFinished ? TEXT("True") : TEXT("False"));
+	PropertiesInEffect.ToString(Out);
+}
+
+bool FChaosPathedMovementState::ShouldReconcile(const FMoverDataStructBase& AuthorityState) const
+{
+	const FChaosPathedMovementState& TypedAuthority = static_cast<const FChaosPathedMovementState&>(AuthorityState);
+
+	static constexpr float PathFractionTolerance = 0.05f;
+
+	if (PropertiesInEffect != TypedAuthority.PropertiesInEffect
+		|| LatestMovementStartFrame != TypedAuthority.LatestMovementStartFrame
+		|| bIsPathFractionIncreasing != TypedAuthority.bIsPathFractionIncreasing
+		|| bHasFinished != TypedAuthority.bHasFinished
+		|| !FMath::IsNearlyEqual(CurrentPathFraction, TypedAuthority.CurrentPathFraction, PathFractionTolerance))
+	{
+		return true;
+	}
+
+	return false;
+}
+
+void FChaosPathedMovementState::Interpolate(const FMoverDataStructBase& From, const FMoverDataStructBase& To, float Pct)
+{
+	const FChaosPathedMovementState& TypedFrom = static_cast<const FChaosPathedMovementState&>(From);
+	const FChaosPathedMovementState& TypedTo = static_cast<const FChaosPathedMovementState&>(To);
+
+	PropertiesInEffect = (Pct < .5f) ? TypedTo.PropertiesInEffect : TypedTo.PropertiesInEffect;
+	LatestMovementStartFrame = (Pct < .5f) ? TypedTo.LatestMovementStartFrame : TypedTo.LatestMovementStartFrame;
+	CurrentPathFraction = FMath::Lerp(TypedFrom.CurrentPathFraction, TypedTo.CurrentPathFraction, Pct);
+	bIsPathFractionIncreasing = (Pct < .5f) ? TypedFrom.bIsPathFractionIncreasing : TypedTo.bIsPathFractionIncreasing;
+	bHasFinished = (Pct < .5f) ? TypedFrom.bHasFinished : TypedTo.bHasFinished;
+}
+
+FMoverDataStructBase* FChaosPathedMovementModeDebugData::Clone() const
+{
+	FChaosPathedMovementModeDebugData* CopyPtr = new FChaosPathedMovementModeDebugData(*this);
+	return CopyPtr;
+}
+
+UScriptStruct* FChaosPathedMovementModeDebugData::GetScriptStruct() const
+{
+	return StaticStruct();
+}
diff --git a/Engine/Plugins/Experimental/ChaosMover/Source/ChaosMover/Private/PathedMovement/ChaosPointMovementPathPattern.cpp b/Engine/Plugins/Experimental/ChaosMover/Source/ChaosMover/Private/PathedMovement/ChaosPointMovementPathPattern.cpp
new file mode 100644
--- /dev/null
+++ b/Engine/Plugins/Experimental/ChaosMover/Source/ChaosMover/Private/PathedMovement/ChaosPointMovementPathPattern.cpp
@@ -0,0 +1,131 @@
+// Copyright Epic Games, Inc. All Rights Reserved.
+
+#include "ChaosMover/PathedMovement/ChaosPointMovementPathPattern.h"
+
+#include "DebugRenderSceneProxy.h"
+#include "Kismet/KismetMathLibrary.h"
+#include "ChaosMover/ChaosMoverSimulation.h"
+#include "Framework/Threading.h"
+//#include "ChaosMover/PathedMovement/ChaosPathedPhysicsDebugDrawComponent.h"
+
+void UChaosPointMovementPathPattern::InitializePattern(UChaosMoverSimulation* InSimulation)
+{
+	Super::InitializePattern(InSimulation);
+
+	RefreshAssignedPointProgress(true);
+}
+
+#if WITH_EDITOR
+void UChaosPointMovementPathPattern::PostEditChangeProperty(struct FPropertyChangedEvent& PropertyChangedEvent)
+{
+	Super::PostEditChangeProperty(PropertyChangedEvent);
+
+	RefreshAssignedPointProgress(true);
+}
+#endif
+
+
+/*
+void UChaosPointMovementPathPattern::AppendDebugDrawElements(UChaosPathedPhysicsDebugDrawComponent& DebugDrawComp, FBoxSphereBounds::Builder& InOutDebugBoundsBuilder)
+{
+	if (!PathPoints.IsEmpty())
+	{
+		const FTransform& PathOriginTransform = GetPathedMoverComp().GetPathOriginTransform();
+		FVector PrevPointLocation = PathOriginTransform.GetLocation();
+		for (const FChaosPointMovementPathPoint& PathPoint : PathPoints)
+		{
+			DebugDrawComp.DebugSpheres.Emplace(4.f, PrevPointLocation, PatternDebugDrawColor);
+
+			FVector PointLocation = PathPoint.WorldLoc;
+			if (PathPoint.Basis == EChaosPointMovementLocationBasis::PreviousPoint)
+			{
+				PointLocation = PrevPointLocation + PathPoint.Location;	
+			}
+			else if (PathPoint.Basis == EChaosPointMovementLocationBasis::PathOrigin)
+			{
+				PointLocation = PathOriginTransform.TransformPositionNoScale(PathPoint.Location);
+			}
+			
+			DebugDrawComp.DebugLines.Emplace(PrevPointLocation, PointLocation, PatternDebugDrawColor, 1.f);
+			InOutDebugBoundsBuilder += PointLocation;
+			
+			PrevPointLocation = PointLocation;
+		}
+
+		DebugDrawComp.DebugSpheres.Emplace(4.f, PrevPointLocation, PatternDebugDrawColor);
+	}
+}
+*/
+
+FTransform UChaosPointMovementPathPattern::CalcUnmaskedTargetRelativeTransform(float PatternProgress, const FTransform& CurTargetTransform) const
+{
+	// The destination point is the first point at or equal to the target progress
+	int32 DestPointIdx = 0;
+	for (; DestPointIdx < PathPoints.Num() && PathPoints[DestPointIdx].Progress < PatternProgress; ++DestPointIdx) {}
+
+	if (PathPoints.IsValidIndex(DestPointIdx))
+	{
+		const FTransform& PathBasis = Simulation ? Simulation->GetPathedMovementBasis() : FTransform::Identity;
+		
+		const FChaosPointMovementPathPoint& NextPoint = PathPoints[DestPointIdx];
+		FChaosPointMovementPathPoint PrevPoint;
+		if (DestPointIdx > 0)
+		{
+			PrevPoint = PathPoints[DestPointIdx - 1];
+		}
+		else
+		{
+			PrevPoint.WorldLoc = PathBasis.GetLocation();
+		}
+
+		const float ProgressSinceLastPoint = PatternProgress - PrevPoint.Progress;
+		const float Alpha = (PatternProgress == PrevPoint.Progress) ? 0.0f : ProgressSinceLastPoint / (NextPoint.Progress - PrevPoint.Progress);
+		
+		// Point locations are calculated in world space, but path targets need to be provided as relative transforms
+		const FVector TargetWorldLoc = UKismetMathLibrary::VLerp(PrevPoint.WorldLoc, NextPoint.WorldLoc, Alpha);
+		const FVector TargetRelativeLoc = PathBasis.InverseTransformPositionNoScale(TargetWorldLoc);
+
+		return FTransform(TargetRelativeLoc);
+	}
+	return FTransform::Identity;
+}
+
+void UChaosPointMovementPathPattern::RefreshAssignedPointProgress(bool bForceRefresh) const
+{
+	Chaos::EnsureIsInGameThreadContext();
+
+	if (!bHasAssignedPointProgress || bForceRefresh)
+	{
+		TotalPathDistance = 0.f;
+		
+		const FTransform& PathBasis = Simulation ? Simulation->GetPathedMovementBasis() : FTransform::Identity;
+		FVector PrevPointLocation = PathBasis.GetLocation();
+
+		// Run through and establish point locations and distance
+		for (const FChaosPointMovementPathPoint& PathPoint : PathPoints)
+		{
+			if (PathPoint.Basis == EChaosPointMovementLocationBasis::PreviousPoint)
+			{
+				PathPoint.WorldLoc = PrevPointLocation + PathPoint.Location;	
+			}
+			else if (PathPoint.Basis == EChaosPointMovementLocationBasis::PathOrigin)
+			{
+				PathPoint.WorldLoc = PathBasis.TransformPositionNoScale(PathPoint.Location);
+			}
+			
+			const float DistanceFromPrevious = (PathPoint.WorldLoc - PrevPointLocation).Size();
+			TotalPathDistance += DistanceFromPrevious;
+			PathPoint.DistanceFromStart = TotalPathDistance;
+
+			PrevPointLocation = PathPoint.WorldLoc;
+		}
+
+		// Go through again and assign progress based on distance
+		for (const FChaosPointMovementPathPoint& PathPoint : PathPoints)
+		{
+			PathPoint.Progress = UKismetMathLibrary::SafeDivide(PathPoint.DistanceFromStart, TotalPathDistance);
+		}
+		
+		bHasAssignedPointProgress = true;
+	}
+}
diff --git a/Engine/Plugins/Experimental/ChaosMover/Source/ChaosMover/Private/PathedMovement/ChaosSplineMovementPathPattern.cpp b/Engine/Plugins/Experimental/ChaosMover/Source/ChaosMover/Private/PathedMovement/ChaosSplineMovementPathPattern.cpp
new file mode 100644
--- /dev/null
+++ b/Engine/Plugins/Experimental/ChaosMover/Source/ChaosMover/Private/PathedMovement/ChaosSplineMovementPathPattern.cpp
@@ -0,0 +1,55 @@
+// Copyright Epic Games, Inc. All Rights Reserved.
+
+#include "ChaosMover/PathedMovement/ChaosSplineMovementPathPattern.h"
+#include "MoverComponent.h"
+#include "ChaosMover/PathedMovement/ChaosPathedMovementMode.h"
+#include "Framework/Threading.h"
+
+#include "Components/SplineComponent.h"
+
+void UChaosSplineMovementPathPattern::InitializePattern(UChaosMoverSimulation* InSimulation)
+{
+	Super::InitializePattern(InSimulation);
+
+	Chaos::EnsureIsInGameThreadContext();
+
+	if (!SplineComp && InSimulation)
+	{
+		UChaosPathedMovementMode& MovementMode = GetMovementMode();
+		if (UMoverComponent* MoverComponent = MovementMode.GetMoverComponent())
+		{
+			if (AActor* OwningActor = MoverComponent->GetOwner(); ensure(OwningActor))
+			{
+				SplineComp = Cast<USplineComponent>(SplineComponentRef.GetComponent(OwningActor));
+				if (!SplineComp)
+				{
+					SplineComp = OwningActor->FindComponentByClass<USplineComponent>();
+				}
+			}
+		}
+	}
+}
+
+FTransform UChaosSplineMovementPathPattern::CalcUnmaskedTargetRelativeTransform(float PatternProgress, const FTransform& CurTargetTransform) const
+{
+	if (SplineComp)
+	{
+		const float SplineLength = SplineComp->GetSplineLength();
+
+		const bool bHasValidBounds = HasValidBounds();
+		const float SplineStartDist = bHasValidBounds ? SplineLength * LowerBound : 0.f;
+		const float SplineEndDist = bHasValidBounds ? SplineLength * UpperBound : 1.f;
+
+		const float DistanceAlongSpline = FMath::Lerp(SplineStartDist, SplineEndDist, PatternProgress);
+		FTransform TargetToSpline = SplineComp->GetTransformAtDistanceAlongSpline(DistanceAlongSpline, ESplineCoordinateSpace::Local, bApplySplineScaling);
+
+		if (!bOrientComponentToPath)
+		{
+			TargetToSpline.SetRotation(FQuat::Identity);
+		}
+		
+		// Intentionally not accounting for any offset between the path origin and the spline's transform - zero progress should always equate to the path origin
+		return TargetToSpline;
+	}
+	return FTransform::Identity;
+}
\ No newline at end of file
diff --git a/Engine/Plugins/Experimental/ChaosMover/Source/ChaosMover/Private/PathedMovement/Transitions/ChaosPathedMovementReachedEndTransition.cpp b/Engine/Plugins/Experimental/ChaosMover/Source/ChaosMover/Private/PathedMovement/Transitions/ChaosPathedMovementReachedEndTransition.cpp
new file mode 100644
--- /dev/null
+++ b/Engine/Plugins/Experimental/ChaosMover/Source/ChaosMover/Private/PathedMovement/Transitions/ChaosPathedMovementReachedEndTransition.cpp
@@ -0,0 +1,46 @@
+// Copyright Epic Games, Inc. All Rights Reserved.
+
+#include "ChaosMover/PathedMovement/Transitions/ChaosPathedMovementReachedEndTransition.h"
+
+#include "ChaosMover/PathedMovement/ChaosPathedMovementTypes.h"
+#include "ChaosMover/ChaosMoverSimulation.h"
+#include "MoverSimulationTypes.h"
+
+#include UE_INLINE_GENERATED_CPP_BY_NAME(ChaosPathedMovementReachedEndTransition)
+
+UChaosPathedMovementReachedEndTransition::UChaosPathedMovementReachedEndTransition(const FObjectInitializer& ObjectInitializer)
+	: Super(ObjectInitializer)
+{
+	bSupportsAsync = true;
+
+	TransitionToMode = DefaultModeNames::Falling;
+}
+
+FTransitionEvalResult UChaosPathedMovementReachedEndTransition::Evaluate_Implementation(const FSimulationTickParams& Params) const
+{
+	FTransitionEvalResult EvalResult;
+
+	if (TransitionToMode.IsNone())
+	{
+		return EvalResult;
+	}
+	
+	check(Simulation);
+
+	bool IsInPathedMovementMode = (nullptr != Cast<IChaosPathedMovementModeInterface>(Simulation->GetCurrentMovementMode()));
+
+	if (IsInPathedMovementMode)
+	{
+		const FChaosPathedMovementState* FoundPreSimMoveState = Params.StartState.SyncState.SyncStateCollection.FindDataByType<FChaosPathedMovementState>();
+		if (FoundPreSimMoveState && FoundPreSimMoveState->bHasFinished)
+		{
+			EvalResult.NextMode = TransitionToMode;
+		}
+	}
+
+	return EvalResult;
+}
+
+void UChaosPathedMovementReachedEndTransition::Trigger_Implementation(const FSimulationTickParams& Params)
+{
+}
\ No newline at end of file
diff --git a/Engine/Plugins/Experimental/ChaosMover/Source/ChaosMover/Public/ChaosMover/Backends/ChaosMoverBackend.h b/Engine/Plugins/Experimental/ChaosMover/Source/ChaosMover/Public/ChaosMover/Backends/ChaosMoverBackend.h
--- a/Engine/Plugins/Experimental/ChaosMover/Source/ChaosMover/Public/ChaosMover/Backends/ChaosMoverBackend.h
+++ b/Engine/Plugins/Experimental/ChaosMover/Source/ChaosMover/Public/ChaosMover/Backends/ChaosMoverBackend.h
@@ -7,6 +7,7 @@
 #include "Chaos/Declares.h"
 #include "ChaosMover/ChaosMoverSimulation.h"
 #include "ChaosMover/ChaosMoverSimulationTypes.h"
+#include "ChaosMover/PathedMovement/ChaosPathedMovementTypes.h"
 #include "Components/ActorComponent.h"
 
 #include "ChaosMoverBackend.generated.h"
@@ -36,44 +37,30 @@ public:
 	CHAOSMOVER_API virtual void BeginPlay() override;
 	CHAOSMOVER_API virtual void EndPlay(const EEndPlayReason::Type EndPlayReason) override;
 
-	CHAOSMOVER_API virtual bool IsAsync() const override
-	{
-		return true;
-	}
-
+	// IMoverBackendLiaisonInterface
+	CHAOSMOVER_API virtual bool IsAsync() const override;
 	// Unused IMoverBackendLiaisonInterface API
-	CHAOSMOVER_API virtual float GetCurrentSimTimeMs()
-	{
-		return 0.0f;
-	}
-	CHAOSMOVER_API virtual int32 GetCurrentSimFrame()
-	{
-		return 0;
-	}
+	CHAOSMOVER_API virtual float GetCurrentSimTimeMs();
+	CHAOSMOVER_API virtual int32 GetCurrentSimFrame();
 
 	CHAOSMOVER_API virtual void ProduceInputData(int32 PhysicsStep, int32 NumSteps, const FMoverTimeStep& TimeStep, UE::ChaosMover::FSimulationInputData& InputData);
 	CHAOSMOVER_API virtual void ConsumeOutputData(const FMoverTimeStep& TimeStep, const UE::ChaosMover::FSimulationOutputData& OutputData);
 	CHAOSMOVER_API virtual void FinalizeFrame(float ResultsTimeInMs);
 
-	UChaosMoverSimulation* GetSimulation()
-	{
-		return Simulation;
-	}
-
-	const UChaosMoverSimulation* GetSimulation() const
-	{
-		return Simulation;
-	}
+	UChaosMoverSimulation* GetSimulation();
+	const UChaosMoverSimulation* GetSimulation() const;
 
 protected:
 	CHAOSMOVER_API virtual void InitSimulation();
 	CHAOSMOVER_API virtual void DeinitSimulation();
 	CHAOSMOVER_API virtual void CreatePhysics();
 	CHAOSMOVER_API virtual void DestroyPhysics();
+
 	CHAOSMOVER_API void CreateCharacterGroundConstraint();
 	CHAOSMOVER_API void DestroyCharacterGroundConstraint();
-	CHAOSMOVER_API void CreatePathTargetConstraint();
-	CHAOSMOVER_API void DestroyPathTargetConstraint();
+
+	CHAOSMOVER_API void CreateActuationConstraint();
+	CHAOSMOVER_API void DestroyActuationConstraint();
 
 	CHAOSMOVER_API virtual void GenerateInput(const FMoverTimeStep& TimeStep, UE::ChaosMover::FSimulationInputData& InputData);
 
@@ -105,10 +92,10 @@ protected:
 	// Character ground constraint, for moving on ground like characters
 	TUniquePtr<Chaos::FCharacterGroundConstraint> CharacterGroundConstraint;
 
-	// Path target constraint, for moving along a path
-	FPhysicsConstraintHandle PathTargetConstraintHandle;
-	FConstraintInstance PathTargetConstraintInstance;
-	FPhysicsUserData PathTargetConstraintPhysicsUserData;
+	// General purpose joint constraint, for moving the controlled component physically
+	FPhysicsConstraintHandle ActuationConstraintHandle;
+	FConstraintInstance ActuationConstraintInstance;
+	FPhysicsUserData ActuationConstraintPhysicsUserData;
 
 	bool bIsUsingAsyncPhysics = false;
 	bool bWantsDestroySim = false;
diff --git a/Engine/Plugins/Experimental/ChaosMover/Source/ChaosMover/Public/ChaosMover/ChaosMoverSettings.h b/Engine/Plugins/Experimental/ChaosMover/Source/ChaosMover/Public/ChaosMover/ChaosMoverSettings.h
new file mode 100644
--- /dev/null
+++ b/Engine/Plugins/Experimental/ChaosMover/Source/ChaosMover/Public/ChaosMover/ChaosMoverSettings.h
@@ -0,0 +1,30 @@
+// Copyright Epic Games, Inc. All Rights Reserved.
+
+#pragma once
+
+#include "Engine/DeveloperSettings.h"
+
+#include "ChaosMoverSettings.generated.h"
+
+/** 
+ * Settings for the Chaos Mover Plugin
+ */
+UCLASS(Config = Engine, DefaultConfig, meta = (DisplayName = "Chaos Mover"), MinimalAPI)
+class UChaosMoverSettings : public UDeveloperSettings
+{
+	GENERATED_BODY()
+
+public:
+
+	CHAOSMOVER_API UChaosMoverSettings();
+
+	/** In the Pathed Movement controller, how far in the future to delay the change in movement behavior 
+	* 	for scheduled changes, leaving enough time to all end points to receive the change before
+	*   it is scheduled to take effect, and therefore avoid corrections.
+	*   There is an option to avoid scheduling altogether, this only applies to scheduled changes.
+	*   To avoid any corrections (but increase latency), schedule any change and set this value to be
+	*   greater than the max round trip time (RTT) supported.
+	*/
+	UPROPERTY(Config, EditAnywhere, Category = "Chaos Mover", meta = (Units = "s", ClampMin = "0.0"))
+	float PathedMovementSchedulingDelaySeconds = 0.3f;
+};
diff --git a/Engine/Plugins/Experimental/ChaosMover/Source/ChaosMover/Public/ChaosMover/ChaosMoverSimulation.h b/Engine/Plugins/Experimental/ChaosMover/Source/ChaosMover/Public/ChaosMover/ChaosMoverSimulation.h
--- a/Engine/Plugins/Experimental/ChaosMover/Source/ChaosMover/Public/ChaosMover/ChaosMoverSimulation.h
+++ b/Engine/Plugins/Experimental/ChaosMover/Source/ChaosMover/Public/ChaosMover/ChaosMoverSimulation.h
@@ -4,9 +4,11 @@
 
 #include "Chaos/Declares.h"
 #include "ChaosMover/ChaosMoverSimulationTypes.h"
+#include "ChaosMover/PathedMovement/ChaosPathedMovementTypes.h"
 #include "ChaosMoverStateMachine.h"
 #include "MoveLibrary/MovementMixer.h"
 #include "MoveLibrary/MoverBlackboard.h"
+#include "MoverTypes.h"
 #include "MoverSimulation.h"
 #include "UObject/WeakObjectPtrTemplates.h"
 #include "PhysicsEngine/ConstraintInstance.h"
@@ -60,8 +62,8 @@ public:
 		TWeakObjectPtr<UImmediateMovementModeTransition> ImmediateModeTransition = nullptr;
 		TWeakObjectPtr<UMovementMixer> MovementMixer = nullptr;
 		Chaos::FCharacterGroundConstraintProxy* CharacterConstraintProxy = nullptr;
-		Chaos::FJointConstraintPhysicsProxy* PathTargetConstraintProxy = nullptr;
-		Chaos::FSingleParticlePhysicsProxy* PathTargetKinematicEndPointProxy = nullptr;
+		Chaos::FJointConstraintPhysicsProxy* ActuationConstraintProxy = nullptr;
+		Chaos::FSingleParticlePhysicsProxy* ActuationConstraintKinematicEndPointProxy = nullptr;
 		Chaos::FConstPhysicsObjectHandle PhysicsObject = nullptr;
 		Chaos::FPhysicsSolver* Solver = nullptr;
 		UWorld* World = nullptr;
@@ -81,6 +83,12 @@ public:
 	CHAOSMOVER_API void ApplyNetStateData(const FMoverSyncState& InNetSyncState);
 	CHAOSMOVER_API void BuildNetStateData(FMoverSyncState& OutNetSyncState) const;
 
+	// Pathed Movement Simulation API
+	// Returns the reference frame of the pathed movement
+	virtual const FTransform& GetPathedMovementBasis() const;
+	// Sets the reference frame of the pathed movement
+	virtual void SetPathedMovementBasis(const FTransform& InPathedMovementBasis);
+
 	//~ Debugging Util functions
 	// Collection for holding extra debug data, that will be sent to the Chaos Visual Debugger for debugging
 	UFUNCTION(BlueprintPure, Category = Mover)
@@ -101,7 +109,7 @@ protected:
 	CHAOSMOVER_API virtual void PostSimulationTickCharacterConstraint(const IChaosCharacterConstraintMovementModeInterface& CharacterConstraintMode, const FMoverTimeStep& TimeStep, UE::ChaosMover::FSimulationOutputData& OutputData);
 
 	// Pathed movement sim steps (kinematic or constrained)
-	CHAOSMOVER_API virtual void PostSimulationTickPathedMovement(const IChaosPathedMovementModeInterface& ConstraintMode, const FMoverTimeStep& TimeStep, UE::ChaosMover::FSimulationOutputData& OutputData);
+	CHAOSMOVER_API virtual void PostSimulationTickMovementActuation(const IChaosMovementActuationInterface& ConstraintMode, const FMoverTimeStep& TimeStep, UE::ChaosMover::FSimulationOutputData& OutputData);
 
 	CHAOSMOVER_API Chaos::FPBDRigidParticleHandle* GetControlledParticle() const;
 
@@ -119,10 +127,12 @@ protected:
 	CHAOSMOVER_API void EnableCharacterConstraint();
 	CHAOSMOVER_API void DisableCharacterConstraint();
 
-	// Path target constraint
-	CHAOSMOVER_API void EnablePathTargetConstraint();
-	CHAOSMOVER_API void DisablePathTargetConstraint();
-
+	// General purpose actuation joint constraint
+	CHAOSMOVER_API void EnableActuationConstraint();
+	CHAOSMOVER_API void DisableActuationConstraint();
+	void SetActuationTargetTransform(const FTransform& TargetTransform);
+	void TeleportActuationTarget(const FTransform& TargetTransform, bool AlsoTeleportControlledParticle = false);
+	
 	CHAOSMOVER_API void TraceMoverData(const UE::ChaosMover::FSimulationOutputData& OutputData);
 
 	// State structs
@@ -143,12 +153,15 @@ protected:
 	// Controlled physics object
 	Chaos::FConstPhysicsObjectHandle PhysicsObject = nullptr;
 
-	// Character ground constraint
+	// Character ground constraint, specifically for moving on ground like characters
 	Chaos::FCharacterGroundConstraintProxy* CharacterConstraintProxy = nullptr;
 
-	// Path target constraint, for movement following a path in a physics based manner
-	Chaos::FJointConstraintPhysicsProxy* PathTargetConstraintProxy = nullptr;
-	Chaos::FSingleParticlePhysicsProxy* PathTargetKinematicEndPointProxy = nullptr;
+	// General purpose actuation joint constraint
+	Chaos::FJointConstraintPhysicsProxy* ActuationConstraintProxy = nullptr;
+	Chaos::FSingleParticlePhysicsProxy* ActuationConstraintKinematicEndPointProxy = nullptr;
+	
+	// Basis transform for pathed movement
+	FTransform PathBasis;
 
 	Chaos::FPhysicsSolver* Solver = nullptr;
 	UWorld* World = nullptr;
diff --git a/Engine/Plugins/Experimental/ChaosMover/Source/ChaosMover/Public/ChaosMover/ChaosMoverSimulationTypes.h b/Engine/Plugins/Experimental/ChaosMover/Source/ChaosMover/Public/ChaosMover/ChaosMoverSimulationTypes.h
--- a/Engine/Plugins/Experimental/ChaosMover/Source/ChaosMover/Public/ChaosMover/ChaosMoverSimulationTypes.h
+++ b/Engine/Plugins/Experimental/ChaosMover/Source/ChaosMover/Public/ChaosMover/ChaosMoverSimulationTypes.h
@@ -158,26 +158,6 @@ public:
 	virtual void UpdateConstraintSettings(Chaos::FCharacterGroundConstraintSettings& ConstraintSettings) const = 0;
 };
 
-// Interface for mover modes moving on a path, with a constraint or kinematically
-UINTERFACE(MinimalAPI)
-class UChaosPathedMovementModeInterface : public UInterface
-{
-	GENERATED_BODY()
-};
-
-class IChaosPathedMovementModeInterface
-{
-	GENERATED_BODY()
-
-public:
-
-	// Whether the simulation should use a constraint to move the actor along the path, or moving the actor kinematically
-	virtual bool ShouldUseConstraint() const = 0;
-
-	// Constraint settings of the joint constraint, only valid when ShouldUseConstraint() is true
-	virtual const Chaos::FPBDJointSettings& GetConstraintSettings() const = 0;
-};
-
 UENUM(BlueprintType)
 enum class EChaosMoverVelocityEffectMode : uint8
 {
diff --git a/Engine/Plugins/Experimental/ChaosMover/Source/ChaosMover/Public/ChaosMover/Character/ChaosCharacterMoverComponent.h b/Engine/Plugins/Experimental/ChaosMover/Source/ChaosMover/Public/ChaosMover/Character/ChaosCharacterMoverComponent.h
--- a/Engine/Plugins/Experimental/ChaosMover/Source/ChaosMover/Public/ChaosMover/Character/ChaosCharacterMoverComponent.h
+++ b/Engine/Plugins/Experimental/ChaosMover/Source/ChaosMover/Public/ChaosMover/Character/ChaosCharacterMoverComponent.h
@@ -22,6 +22,7 @@ public:
 
 	CHAOSMOVER_API virtual bool TryGetFloorCheckHitResult(FHitResult& OutHitResult) const override;
 	CHAOSMOVER_API virtual void ProduceInput(const int32 DeltaTimeMS, FMoverInputCmdContext* Cmd) override;
+	CHAOSMOVER_API virtual void DoQueueNextMode(FName DesiredModeName, bool bShouldReenter=false) override;
 
 	// Broadcast when this actor lands on a valid surface.
 	UPROPERTY(BlueprintAssignable, Category = "Chaos Mover")
@@ -36,9 +37,14 @@ protected:
 	CHAOSMOVER_API virtual void ProcessSimulationEvent(const FMoverSimulationEventData& EventData) override;
 	CHAOSMOVER_API virtual void SetAdditionalSimulationOutput(const FMoverDataCollection& Data) override;
 
+	CHAOSMOVER_API void ClearQueuedMode();
+
 	bool bFloorResultSet = false;
 	FFloorCheckResult LatestFloorResult;
 
 	FVector LaunchVelocityOrImpulse = FVector::ZeroVector;
 	EChaosMoverVelocityEffectMode LaunchMode = EChaosMoverVelocityEffectMode::AdditiveVelocity;
+
+	// Queued immediate mode transition, it will be transmitted to the PT next time ProduceInput is called
+	FName QueuedModeTransitionName = NAME_None;
 };
diff --git a/Engine/Plugins/Experimental/ChaosMover/Source/ChaosMover/Public/ChaosMover/PathedMovement/ChaosPathedMovementControllerComponent.h b/Engine/Plugins/Experimental/ChaosMover/Source/ChaosMover/Public/ChaosMover/PathedMovement/ChaosPathedMovementControllerComponent.h
new file mode 100644
--- /dev/null
+++ b/Engine/Plugins/Experimental/ChaosMover/Source/ChaosMover/Public/ChaosMover/PathedMovement/ChaosPathedMovementControllerComponent.h
@@ -0,0 +1,176 @@
+// Copyright Epic Games, Inc. All Rights Reserved.
+
+#pragma once
+
+#include "CoreMinimal.h"
+#include "Delegates/Delegate.h"
+#include "ChaosMover/ChaosMoverSimulationTypes.h"
+#include "ChaosMover/PathedMovement/ChaosPathedMovementTypes.h"
+
+#include "ChaosPathedMovementControllerComponent.generated.h"
+
+namespace Chaos
+{
+class FPBDRigidsSolver;
+}
+class UMoverComponent;
+
+// Fired on the game thread when it has received notification that the pathed movement has started
+DECLARE_DYNAMIC_MULTICAST_DELEGATE(FChaosMover_OnPathedMovementStarted);
+// Fired on the game thread when it has received notification that the pathed movement has stopped
+DECLARE_DYNAMIC_MULTICAST_DELEGATE_OneParam(FChaosMover_OnPathedMovementStopped, bool, bReachedEndOfPlayback);
+// Fired on the game thread when it has received notification that the pathed movement has bounced (reached the end of one leg in a round trip)
+DECLARE_DYNAMIC_MULTICAST_DELEGATE(FChaosMover_OnPathedMovementBounced);
+// Fired on the game thread when reverse playback has changed
+DECLARE_DYNAMIC_MULTICAST_DELEGATE_OneParam(FChaosMover_OnPathedMovementReversePlaybackChanged, bool, bIsReversePlayback);
+// Fired on the game thread when looping playback has changed
+DECLARE_DYNAMIC_MULTICAST_DELEGATE_OneParam(FChaosMover_OnPathedMovementLoopingPlaybackChanged, bool, bIsLoopingPlayback);
+// Fired on the game thread when one way playback has changed
+DECLARE_DYNAMIC_MULTICAST_DELEGATE_OneParam(FChaosMover_OnPathedMovementOneWayPlaybackChanged, bool, bIsOneWayPlayback);
+
+/** Execution type for pathed movement modifications, i.e. whether the modification is executed on the authority only or predicted on the client */
+UENUM(BlueprintType)
+enum class EChaosPathedMovementExecutionType : uint8
+{
+	// Only executed if attempted on the authority.
+	// If attempted on non authoritative client, this action will be ignored.
+	// This will execute on an authoritative client (client only actors)
+	AuthorityOnly UMETA(DisplayName = "Authority Only"),
+	// If the actor is autonomous proxy, this action will execute predictively.
+	// This will also result in an input that will cause the action to execute on the server and replicate to sim proxies.
+	// If the actor is simulated proxy, this will be ignored.
+	// On the server, this will be ignored, relying on input replication instead.
+	// This will execute on an authoritative client (client only actors)
+	ClientPredicted_AutonomousOnly UMETA(DisplayName = "Client Predicted - Autonomous Only"),
+};
+
+/** 
+*   This controller component gives control over pathed movement of an actor with a mover component and an active pathed movement mode.
+*   You can start/stop, play forward or in reverse, play once or loop and play one way (0 to 1) or round trip (0 to 1 to 0).
+*   In async mode, control requests are sent as scheduled to take effect at a future frame, to give enough time for all end points
+*   to receive the request and react at the same exact scheduled frame.
+*/
+UCLASS(BlueprintType, Blueprintable, meta = (BlueprintSpawnableComponent))
+class UChaosPathedMovementControllerComponent : public UActorComponent
+	, public IMoverInputProducerInterface
+{
+	GENERATED_BODY()
+
+public:
+	UChaosPathedMovementControllerComponent();
+
+	// UActorComponent
+	virtual void BeginPlay() override;
+
+	// IMoverInputProducerInterface
+	virtual void ProduceInput_Implementation(int32 SimTimeMs, FMoverInputCmdContext& InputCmdResult) override;
+
+	// Pathed Movement Control Layer
+	// Start/Stop Playing Path
+	/** Request to start or resume moving along the path. If a previous non looping playback reached the end naturally, this will restart playback.
+	*   @param ExecutionType The execution type of this action @see EChaosPathedMovementExecutionType
+	*   @param bIsScheduled If true, the execution will be delayed to leave enough time for all end points to receive request and start at the same frame to avoid a correction. @see UChaosMoverSettings::PathedMovementSchedulingDelaySeconds
+	*/
+	UFUNCTION(BlueprintCallable, Category = "Pathed Movement")
+	CHAOSMOVER_API void RequestStartPlayingPath(EChaosPathedMovementExecutionType ExecutionType = EChaosPathedMovementExecutionType::AuthorityOnly, bool bIsScheduled = true);
+	/** Request to stop playback. */
+	UFUNCTION(BlueprintCallable, Category = "Pathed Movement")
+	CHAOSMOVER_API void RequestStopPlayingPath(EChaosPathedMovementExecutionType ExecutionType = EChaosPathedMovementExecutionType::AuthorityOnly, bool bIsScheduled = true);
+	/** Do we want to be starting or resuming playback? */
+	UFUNCTION(BlueprintCallable, Category = "Pathed Movement")
+	CHAOSMOVER_API bool WantsPlayingPath() const;
+	/** Is a playback under way (not just requested) or stopped? */
+	UFUNCTION(BlueprintCallable, Category = "Pathed Movement")
+	CHAOSMOVER_API bool IsPlayingPath() const;
+	/** Broadcast when the pathed movement has started */
+	UPROPERTY(BlueprintAssignable, Category = "Pathed Movement")
+	FChaosMover_OnPathedMovementStarted OnPathedMovementStarted;
+	/** Broadcast when the pathed movement has stopped */
+	UPROPERTY(BlueprintAssignable, Category = "Pathed Movement")
+	FChaosMover_OnPathedMovementStopped OnPathedMovementStopped;
+	/** Broadcast when the pathed movement has bounced (reached the end of one leg in a round trip) */
+	UPROPERTY(BlueprintAssignable, Category = "Pathed Movement")
+	FChaosMover_OnPathedMovementBounced OnPathedMovementBounced;
+	// Reverse Playback
+	/** Request to set playback in reverse (bWantsReversePlayback = True) or forward (False) */
+	UFUNCTION(BlueprintCallable, Category = "Pathed Movement")
+	CHAOSMOVER_API void RequestReversePlayback(bool bWantsReversePlayback, EChaosPathedMovementExecutionType ExecutionType = EChaosPathedMovementExecutionType::AuthorityOnly, bool bIsScheduled = true);
+	/** Do we want the playback to be in reverse? */
+	UFUNCTION(BlueprintCallable, Category = "Pathed Movement")
+	CHAOSMOVER_API bool WantsReversePlayback() const;
+	/** Is the playback in reverse */
+	UFUNCTION(BlueprintCallable, Category = "Pathed Movement")
+	CHAOSMOVER_API bool IsReversePlayback() const;
+	/** Broadcast when reverse has changed */
+	UPROPERTY(BlueprintAssignable, Category = "Pathed Movement")
+	FChaosMover_OnPathedMovementReversePlaybackChanged OnPathedMovementReversePlaybackChanged;
+	/** Broadcast when the pathed movement has stopped */
+	// Looping Playback
+	/** Request the playback to be looping (True) or not (False)*/
+	UFUNCTION(BlueprintCallable, Category = "Pathed Movement")
+	CHAOSMOVER_API void RequestLoopingPlayback(bool bWantsLoopingPlayback, EChaosPathedMovementExecutionType ExecutionType = EChaosPathedMovementExecutionType::AuthorityOnly, bool bIsScheduled = true);
+	/** Do we want the playback to be looping or forward? */
+	UFUNCTION(BlueprintCallable, Category = "Pathed Movement")
+	CHAOSMOVER_API bool WantsLoopingPlayback() const;
+	/** Is the playback looping */
+	UFUNCTION(BlueprintCallable, Category = "Pathed Movement")
+	CHAOSMOVER_API bool IsLoopingPlayback() const;
+	/** Broadcast when looping has changed */
+	UPROPERTY(BlueprintAssignable, Category = "Pathed Movement")
+	FChaosMover_OnPathedMovementLoopingPlaybackChanged OnPathedMovementLoopingPlaybackChanged;
+	// One-way vs. round trip playback
+	/** Request the playback be one way (0->1), or round trip (there and back, 0->1->0) */
+	UFUNCTION(BlueprintCallable, Category = "Pathed Movement")
+	CHAOSMOVER_API void RequestOneWayPlayback(bool bWantsOneWayPlayback, EChaosPathedMovementExecutionType ExecutionType = EChaosPathedMovementExecutionType::AuthorityOnly, bool bIsScheduled = true);
+	UFUNCTION(BlueprintCallable, Category = "Pathed Movement")
+	/** Do we want the playback to be one way or round trip? */
+	CHAOSMOVER_API bool WantsOneWayPlayback() const;
+	/** Is the playback one way or round trip */
+	UFUNCTION(BlueprintCallable, Category = "Pathed Movement")
+	CHAOSMOVER_API bool IsOneWayPlayback() const;
+	/** Broadcast when one way has changed */
+	UPROPERTY(BlueprintAssignable, Category = "Pathed Movement")
+	FChaosMover_OnPathedMovementOneWayPlaybackChanged OnPathedMovementOneWayPlaybackChanged;
+
+	UPROPERTY(EditAnywhere, Category="Pathed Movement")
+	FChaosMutablePathedMovementProperties InitialProperties;
+
+protected:
+	// Processes output events coming from the simulation and broadcasts delegates accordingly
+	UFUNCTION()
+	virtual void OnPostSimEventReceived(const FMoverSimulationEventData& EventData);
+
+	UFUNCTION()
+	virtual void OnPostFinalize(const FMoverSyncState& SyncState, const FMoverAuxStateContext& AuxState);
+
+private:
+	// This returns the earliest future frame index at which a requested change can occur,
+	// giving enough time for the change to be received on all end points before occurring
+	int32 GetEarliestFrameForChange() const;
+
+	// Schedule NewProperties to take effect at the earliest possible time
+	void ExecuteOrScheduleChange(const FChaosMutablePathedMovementProperties& NewProperties, bool bIsScheduled, bool bIsStart = false);
+
+	// Util function to get the physics solver
+	Chaos::FPhysicsSolver* GetPhysicsSolver() const;
+
+	// Mutable pathed movement properties to take effect when server frame reaches LastChangeFrame
+	FChaosMutablePathedMovementProperties PathedMovementMutableProperties;
+
+	// The (server) frame at which PathedMovementMutableProperties took effect
+	// Note that it can be greater than the current server frame to schedule that change
+	// (the changes are sent immediately but scheduled to take effect in the future)
+	int32 LastChangeFrame = INDEX_NONE;
+
+	// The frame at which movement last started (but this may be in the future for a scheduled start)
+	int32 MovementStartFrame = INDEX_NONE;
+	FTimerHandle PathedMovementDelayedStartTimerHandle;
+
+	// The (possibly interpolated) post simulation state 
+	FChaosPathedMovementState PostSimState;
+
+	// Cached mover component if found on owner actor
+	// We use the OnPostMovement delegate to reflect the current state to gameplay (interpolated state in async)
+	UPROPERTY(Transient)
+	TObjectPtr<UMoverComponent> MoverComponent = nullptr;
+};
diff --git a/Engine/Plugins/Experimental/ChaosMover/Source/ChaosMover/Public/ChaosMover/PathedMovement/ChaosPathedMovementMode.h b/Engine/Plugins/Experimental/ChaosMover/Source/ChaosMover/Public/ChaosMover/PathedMovement/ChaosPathedMovementMode.h
new file mode 100644
--- /dev/null
+++ b/Engine/Plugins/Experimental/ChaosMover/Source/ChaosMover/Public/ChaosMover/PathedMovement/ChaosPathedMovementMode.h
@@ -0,0 +1,131 @@
+// Copyright Epic Games, Inc. All Rights Reserved.
+
+#pragma once
+
+#include "CoreMinimal.h"
+#include "MovementMode.h"
+#include "MoverTypes.h"
+#include "AlphaBlend.h"
+#include "PhysicsEngine/ConstraintInstance.h"
+#include "ChaosMover/ChaosMoverSimulationTypes.h"
+#include "ChaosMover/ChaosMovementMode.h"
+#include "ChaosMover/PathedMovement/ChaosPathedMovementTypes.h"
+#include "Chaos/PBDJointConstraintTypes.h"
+
+#include "ChaosPathedMovementMode.generated.h"
+
+class UChaosPathedMovementPatternBase;
+
+/**
+ * Base for any movement mode that proceeds along a defined path within a specified amount of time
+ */
+UCLASS(Within = MoverComponent)
+class UChaosPathedMovementMode
+	: public UChaosMovementMode
+	, public IChaosMovementActuationInterface
+	, public IChaosPathedMovementModeInterface
+{
+	GENERATED_BODY()
+
+public:
+	UChaosPathedMovementMode(const FObjectInitializer& ObjectInitializer);
+
+	// UBaseMovementMode
+	virtual void OnRegistered(const FName ModeName) override;
+	virtual void GenerateMove_Implementation(const FMoverTickStartData& StartState, const FMoverTimeStep& TimeStep, FProposedMove& OutProposedMove) const override;
+	virtual void SimulationTick_Implementation(const FSimulationTickParams& Params, FMoverTickEndData& OutputState) override;
+
+	// IChaosConstraintMovementModeInterface
+	virtual bool ShouldUseConstraint() const override;
+	virtual const Chaos::FPBDJointSettings& GetConstraintSettings() const override;
+
+	// IChaosProgressingPathedMovementModeInterface
+	virtual Chaos::FRigidTransform3 CalcTargetRelativeTransform(float ProgressAmt) const override;
+
+
+#if WITH_EDITOR
+	// UBaseMovementMode Interface
+	EDataValidationResult IsDataValid(FDataValidationContext& Context) const override;
+#endif // WITH_EDITOR
+	
+	/** Called when the owning UChaosPathedPhysicsMoverComponent is registered (i.e. at both edit and game time) */
+	virtual void InitializePath();
+	
+	/** Finds the first path pattern on this mode of the given type, if there is one */
+	UFUNCTION(BlueprintPure, Category = PathedMovementMode, DisplayName = "Find Pattern", meta=(DeterminesOutputType = PatternType))
+	UChaosPathedMovementPatternBase* BP_FindPattern(TSubclassOf<UChaosPathedMovementPatternBase> PatternType) const;
+	
+	/**
+	 * Set the time it takes to travel one way along this path from start to end.
+	 * Only valid to call in BeginPlay or the Construction Script, ignored otherwise.
+	 * Once play has begun, the path duration can be modified indirectly via SetMovementTimeDilation on the owning PathedMoverComponent
+	 */
+	UFUNCTION(BlueprintCallable, Category = PathedMovementMode, DisplayName = "Set Path Duration (BeginPlay only)")
+	void SetOneWayTripDuration_BeginPlayOnly(float NewDuration);
+	float GetOneWayTripDuration() const { return OneWayPlaybackDuration; }
+	
+protected:
+	void OnPlaybackStarting(float BaseSimTimeMs);
+	void OnPlaybackStopping(float BaseSimTimeMs, bool bReachedEndOfPlayback);
+	void OnPlaybackBounced(float BaseSimTimeMs);
+	void OnInReverseChanging(float BaseSimTimeMs, bool bNewInReverse);
+	void OnLoopingChanging(float BaseSimTimeMs, bool bNewLooping);
+	void OnOneWayChanging(float BaseSimTimeMs, bool bNewOneWay);
+
+	void SetDefaultJointConstraintProperties();
+	
+private:
+	/**
+	 * True to have the updated component be able to respond to physics inputs by being "pulled" along the path using a joint constraint
+	 * configured using JointConstraintProperties below.
+	 * 
+	 * If false, no joint will be used, and the affected component will be moved kinematically (i.e. independent of other physics effects and treated as though it has
+	 * infinite mass in the physics sim) and never deviate from the path.
+	 *
+	 * Note that the movement is calculated on the physics thread either way, so it will correctly interact with other bodies in the sim.
+	 *
+	 * Ex: Platform moving along a path in a straight line. If this is enabled, the platform could wobble, dip, shake, stall, etc. in response to physics forces/impulses
+	 *		(depending on the joint and physics properties of the platform). If this is disabled, the platform will rigidly move along the path and ignore all physics.
+	 */
+	UPROPERTY(EditAnywhere, BlueprintReadOnly, Category = "Pathed Movement Mode", meta = (AllowPrivateAccess))
+	bool bUseJointConstraint = true;
+
+	/**
+	 * Configuration of the joint constraint that (if enabled) pulls the updated component along the path
+	 * The default settings are for a very stiff joint that allows just a little bit of "play"
+	 */
+	UPROPERTY(EditAnywhere, Category = "Pathed Movement Mode", meta = (EditCondition = bUseJointConstraint))
+	FConstraintProfileProperties JointConstraintProperties;
+
+	/** Total duration in seconds of a one-way trip from start to end. */
+	UPROPERTY(EditAnywhere, BlueprintReadOnly, Category = "Pathed Movement Mode", meta = (ClampMin=0, UIMin=0, Units=Seconds, AllowPrivateAccess))
+	float OneWayPlaybackDuration = 5.f;
+
+	//@todo DanH: Try to think of a way to make this less order-dependent/error-prone
+	/**
+	 * The set of patterns to follow that combine to form the path the component will follow
+	 *
+	 * NOTE: Order can matter here!! The patterns are evaluated each physics frame in order starting from the front.
+	 *		 Each successive pattern is aware of the aggregate results from the previous patterns.
+	 *		 In most cases, this doesn't matter, as the pattern calculates the target transform against the static path origin point.
+	 *		 However, some patterns necessarily determine their target based on the calculated target to that point
+	 *		 For example, the LookAtRotationPattern calculates the look-at rotation target based on the target location
+	 */
+	UPROPERTY(EditAnywhere, Instanced, BlueprintReadOnly, Category = "Pathed Movement Mode", meta = (AllowPrivateAccess))
+	TArray<TObjectPtr<UChaosPathedMovementPatternBase>> PathPatterns;
+	
+	/** The kind of easing to apply when traveling along the path */
+	UPROPERTY(EditAnywhere, BlueprintReadOnly, Category = "Pathed Movement Mode", meta = (AllowPrivateAccess))
+	EAlphaBlendOption Easing = EAlphaBlendOption::Linear;
+
+	/** If using a custom ease, this is the curve that will be used. If blank, will fall back to standard linear interpolation. */
+	UPROPERTY(EditAnywhere, BlueprintReadOnly, Category = "Pathed Movement Mode", meta = (EditCondition = "Easing==EAlphaBlendOption::Custom", AllowPrivateAccess))
+	TObjectPtr<UCurveFloat> CustomEasingCurve = nullptr;
+
+	Chaos::FPBDJointSettings JointSettings;
+
+#if !UE_BUILD_SHIPPING
+	FChaosPathedMovementModeDebugData DebugData;
+#endif // !UE_BUILD_SHIPPING
+};
+
diff --git a/Engine/Plugins/Experimental/ChaosMover/Source/ChaosMover/Public/ChaosMover/PathedMovement/ChaosPathedMovementPatternBase.h b/Engine/Plugins/Experimental/ChaosMover/Source/ChaosMover/Public/ChaosMover/PathedMovement/ChaosPathedMovementPatternBase.h
new file mode 100644
--- /dev/null
+++ b/Engine/Plugins/Experimental/ChaosMover/Source/ChaosMover/Public/ChaosMover/PathedMovement/ChaosPathedMovementPatternBase.h
@@ -0,0 +1,143 @@
+// Copyright Epic Games, Inc. All Rights Reserved.
+
+#pragma once
+
+#include "CoreMinimal.h"
+#include "ChaosMover/PathedMovement/ChaosPathedMovementTypes.h"
+
+#include "ChaosPathedMovementPatternBase.generated.h"
+
+struct FPhysicsMoverAsyncInput;
+class UCurveFloat;
+enum class EAlphaBlendOption : uint8;
+class UChaosPathedMovementMode;
+class UChaosMoverSimulation;
+//class UChaosPathedPhysicsDebugDrawComponent;
+
+//@todo DanH: What's the Flags specifier for? Does that do everything now or do I still need the Bitflags meta?
+UENUM(Flags, meta = (Bitflags))
+enum class EChaosPatternAxisMaskFlags : uint8
+{
+	None = 0 UMETA(Hidden),
+
+	X = 1 UMETA(DisplayName = "X / Roll"),
+	Y = 1 << 1 UMETA(DisplayName = "Y / Pitch"),
+	Z = 1 << 2 UMETA(DisplayName = "Z / Yaw)"),
+
+	All = X | Y | Z UMETA(Hidden)
+};
+ENUM_CLASS_FLAGS(EChaosPatternAxisMaskFlags);
+
+UCLASS(Abstract, Within = ChaosPathedMovementMode, DefaultToInstanced, Blueprintable, BlueprintType, EditInlineNew)
+class UChaosPathedMovementPatternBase : public UObject
+{
+	GENERATED_BODY()
+
+public:
+	virtual void InitializePattern(UChaosMoverSimulation* InSimulation);
+	virtual void ProduceInputs_External(OUT FPhysicsMoverAsyncInput& Input);
+
+	const UChaosMoverSimulation* GetSimulation() const
+	{
+		return Simulation;
+	}
+
+	UChaosPathedMovementMode& GetMovementMode() const;
+	
+	FTransform CalcTargetRelativeTransform(float OverallPathProgress, const FTransform& CurTargetTransform) const;
+
+	virtual bool DebugDrawUsingStepSamples() const { return bDebugDrawPattern; }
+	//virtual void AppendDebugDrawElements(UChaosPathedPhysicsDebugDrawComponent& DebugDrawComp, FBoxSphereBounds::Builder& InOutDebugBoundsBuilder) {}
+
+#if WITH_EDITOR
+	//@todo DanH: Pipe this through (I wonder if there should be a UPackage::IsDataValid() override that auto-pipes it to all UObjects in the package?)
+	// virtual EDataValidationResult IsDataValid(class FDataValidationContext& Context) const override;
+
+	friend class FChaosPatternPostChangeMovementModeHelper;
+#endif
+	
+protected:
+	/**
+	* Where patterns calculate their target transforms at the given progress, optionally relative to the current aggregate target.
+	* Axis masking is applied to this result before it's actually used.
+	*/
+	virtual FTransform CalcUnmaskedTargetRelativeTransform(float PatternProgress, const FTransform& CurTargetTransform) const { return FTransform::Identity; }
+
+	//
+	// NOTE: Making any of these properties mutable at game time will require either replicating them (easier, but not as good) or pushing them as inputs to the PT (harder, more correct)
+	//
+
+	/** Along which axes is this pattern disallowed from modifying the translation/location of the updated component? */
+	UPROPERTY(EditAnywhere, Category = "Path Pattern", meta = (Bitmask, BitmaskEnum = "/Script/ChaosMover.EChaosPatternAxisMaskFlags"))
+	uint8 TranslationMasks = (uint8)EChaosPatternAxisMaskFlags::None;
+
+	/** Along which axes is this pattern disallowed from modifying the rotation of the updated component? */
+	UPROPERTY(EditAnywhere, Category = "Path Pattern", meta = (Bitmask, BitmaskEnum = "/Script/ChaosMover.EChaosPatternAxisMaskFlags"))
+	uint8 RotationMasks = (uint8)EChaosPatternAxisMaskFlags::None;
+	
+	/** Along which axes is this pattern disallowed from modifying the scale of the updated component? */
+	UPROPERTY(EditAnywhere, Category = "Path Pattern", meta = (Bitmask, BitmaskEnum = "/Script/ChaosMover.EChaosPatternAxisMaskFlags"))
+	uint8 ScaleMasks = (uint8)EChaosPatternAxisMaskFlags::None;
+
+	/**
+	 * If true, this pattern will not begin to take effect until the previous pattern has completed.
+	 * Note: If true and the previous pattern's EndAtPathProgress is 1, this pattern will never start.
+	 */
+	UPROPERTY(EditAnywhere, Category = "Path Pattern")
+	bool bStartAfterPreviousPattern = false;
+
+	// /**
+	//  * The desired duration (in overall path progress) for this pattern once it begins.
+	//  * If this desired duration is longer than the progress remaining when the pattern starts, it will be clamped to fill what remains
+	//  */
+	// UPROPERTY(EditAnywhere, Category = "Path Pattern", meta = (EditCondition = "bStartAfterPreviousPattern", UIMin = 0.f, ClampMin = 0.f, UIMax = 1.f, ClampMax = 1.f))
+	// float DesiredPatternDuration = 1.f;
+
+	/** The overall path progress when this pattern should begin, where 0 is the start of the path and 1 is the end. Must be less than EndAtProgress. */
+	UPROPERTY(EditAnywhere, Category = "Path Pattern", meta = (EditCondition = "!bStartAfterPreviousPattern", UIMin = 0.f, ClampMin = 0.f, UIMax = 1.f, ClampMax = 1.f))
+	float StartAtPathProgress = 0.f;
+
+	/** The overall path progress when this pattern should complete, where 0 is the start of the path and 1 is the end. Must be greater than StartAtProgress. */
+	UPROPERTY(EditAnywhere, Category = "Path Pattern", meta = (UIMin = 0.f, ClampMin = 0.f, UIMax = 1.f, ClampMax = 1.f))
+	float EndAtPathProgress = 1.f;
+
+	//@todo DanH: Maybe this should be a float and allow partial loops?
+	/**
+	 * The number of loops to complete within the active span of this pattern (i.e. between StartAtProgress and EndAtProgress)
+	 * on a single run along the full aggregate path. Setting to 0 effectively disables this pattern.
+	 */
+	UPROPERTY(EditAnywhere, Category = "Path Pattern", meta = (UIMin = 0, ClampMin = 0))
+	int32 NumLoopsPerPath = 1;
+
+	/** Whether a loop is played one way (0 -> 1) or there and back (0 -> 1 -> 0)*/
+	UPROPERTY(EditAnywhere, Category = "Path Pattern")
+	bool IsOneWay = true;
+	
+	/** 
+	 * If true, the component will be rotated to face in the direction of this pattern's motion.
+	 * To have the component face in the direction of the aggregate path, enable this on all movement patterns.
+	 */
+	UPROPERTY(EditAnywhere, Category = "Path Pattern")
+	bool bOrientComponentToPath = false;
+
+	/** The kind of easing to apply when traveling along the path */
+	UPROPERTY(EditAnywhere, Category = "Path Pattern")
+	EAlphaBlendOption Easing;
+
+	/** If using a custom ease, this is the curve that will be used. If blank, will fall back to standard linear interpolation. */
+	UPROPERTY(EditAnywhere, Category = "Path Pattern", meta = (EditCondition = "Easing==EAlphaBlendOption::Custom"))
+	TObjectPtr<UCurveFloat> CustomEasingCurve = nullptr;
+
+	/** True to draw debug lines for this specific pattern in editor views */
+	UPROPERTY(EditAnywhere, Category = "Path Pattern|Debug")
+	bool bDebugDrawPattern = true;
+
+	//@todo DanH: More central setting for color per pattern index maybe?
+	/** The color used for debug draws of this pattern */
+	UPROPERTY(EditAnywhere, Category = "Path Pattern|Debug", meta = (EditCondition = bDrawDebug))
+	FColor PatternDebugDrawColor;
+
+	UChaosMoverSimulation* Simulation = nullptr;
+
+	//friend class UChaosPathedPhysicsDebugDrawComponent;
+};
diff --git a/Engine/Plugins/Experimental/ChaosMover/Source/ChaosMover/Public/ChaosMover/PathedMovement/ChaosPathedMovementTypes.h b/Engine/Plugins/Experimental/ChaosMover/Source/ChaosMover/Public/ChaosMover/PathedMovement/ChaosPathedMovementTypes.h
new file mode 100644
--- /dev/null
+++ b/Engine/Plugins/Experimental/ChaosMover/Source/ChaosMover/Public/ChaosMover/PathedMovement/ChaosPathedMovementTypes.h
@@ -0,0 +1,336 @@
+// Copyright Epic Games, Inc. All Rights Reserved.
+
+#pragma once
+
+#include "CoreMinimal.h"
+#include "MoverTypes.h"
+#include "Chaos/Core.h"
+#include "MoverSimulationTypes.h"
+
+#include "ChaosPathedMovementTypes.generated.h"
+
+namespace UE::ChaosMover::PathBlackboard
+{
+	const FName TargetRelativeTransform = TEXT("TargetRelativeTransform");
+}
+
+namespace Chaos
+{
+	class FPBDJointSettings;
+}
+
+/**
+ * Properties that can change during game-time and affect path movement calculations.
+ * Split into its own struct because both the input and sync state structs want to hold onto these. 
+ */
+USTRUCT()
+struct FChaosMutablePathedMovementProperties
+{
+	GENERATED_BODY()
+
+	// True = Wants to being playing the path, False = Wants to be stopped
+	UPROPERTY(EditAnywhere, Category = ChaosMover)
+	bool bWantsToPlay = false;
+
+	// True = Reverse, False = Forward
+	UPROPERTY(EditAnywhere, Category = ChaosMover)
+	bool bWantsReversePlayback = false;
+
+	// True = Looping, False = Will stop next time it reaches the end
+	// (that end could be 0 or 1 depending on bWantsOneWayPlayback and bWantsReversePlayback)
+	UPROPERTY(EditAnywhere, Category = ChaosMover)
+	bool bWantsLoopingPlayback = false;
+
+	// True = One way playback (0->1), False = Round trip, there and back (0->1->0)
+	UPROPERTY(EditAnywhere, Category = ChaosMover)
+	bool bWantsOneWayPlayback = false;
+
+	// These properties cannot be mutated yet, but were supported by Dan's implementation. Coming soon.
+	//UPROPERTY(VisibleAnywhere, Category = ChaosMover)
+	//FTransform PathOrigin;
+	//UPROPERTY(VisibleAnywhere, Category = ChaosMover)
+	//bool bIsUsingConstraint;
+
+	void NetSerialize(FArchive& Ar);
+	void ToString(FAnsiStringBuilderBase& Out) const;
+	
+	bool operator==(const FChaosMutablePathedMovementProperties& Other) const;
+	bool operator!=(const FChaosMutablePathedMovementProperties& Other) const { return !operator==(Other); }
+};
+
+/**
+ * Frame inputs for all pathed movement
+ * Any property that can change during game time must be delivered in this way, only immutable properties can be safely
+ * referenced directly on the object they come from (usually the movement mode)
+ */
+USTRUCT()
+struct FChaosPathedMovementInputs : public FMoverDataStructBase
+{
+	GENERATED_BODY()
+
+	UPROPERTY(VisibleAnywhere, Category = ChaosMover)
+	int32 MovementStartFrame = INDEX_NONE;
+
+	UPROPERTY(VisibleAnywhere, Category = ChaosMover)
+	int32 LastChangeFrame = INDEX_NONE;
+
+	UPROPERTY(VisibleAnywhere, Category = ChaosMover)
+	FChaosMutablePathedMovementProperties Props;
+
+	virtual bool NetSerialize(FArchive& Ar, UPackageMap* Map, bool& bOutSuccess) override;
+	virtual FMoverDataStructBase* Clone() const override { return new FChaosPathedMovementInputs(*this); }
+	virtual UScriptStruct* GetScriptStruct() const override { return StaticStruct(); }
+	virtual void ToString(FAnsiStringBuilderBase& Out) const override;
+	virtual bool ShouldReconcile(const FMoverDataStructBase& AuthorityState) const override;
+	virtual void Interpolate(const FMoverDataStructBase& From, const FMoverDataStructBase& To, float Pct) override;
+	virtual void Merge(const FMoverDataStructBase& From) override;
+};
+
+/** State of a pathed movement that persists from frame to frame */
+USTRUCT()
+struct FChaosPathedMovementState : public FMoverDataStructBase
+{
+	GENERATED_BODY()
+
+	/** This is a value between 0 and 1 that can be directly used along the path to calculate
+	  * the current position on it, 0.0 being a the start of the path, 1.0 being at the end.
+	  */
+	UPROPERTY(VisibleAnywhere, Category = ChaosMover)
+	float CurrentPathFraction = 0.0f;
+
+	/** The frame at which we last started playing, or INDEX_NONE if we were never playing */
+	UPROPERTY(VisibleAnywhere, Category = ChaosMover)
+	int32 LatestMovementStartFrame = INDEX_NONE;
+
+	/** Whether traveling from 0 to 1 (true) or 1 to 0 (false) */
+	UPROPERTY(VisibleAnywhere, Category = ChaosMover)
+	bool bIsPathFractionIncreasing = true;
+	
+	/** Whether the previous playback finished by reaching its end point.
+	*/
+	UPROPERTY(VisibleAnywhere, Category = ChaosMover)
+	bool bHasFinished = false;
+	
+	/** The mutable properties in effect */
+	UPROPERTY(VisibleAnywhere, Category = ChaosMover)
+	FChaosMutablePathedMovementProperties PropertiesInEffect;
+
+	virtual bool NetSerialize(FArchive& Ar, UPackageMap* Map, bool& bOutSuccess) override;
+	virtual FMoverDataStructBase* Clone() const override { return new FChaosPathedMovementState(*this); }
+	virtual UScriptStruct* GetScriptStruct() const override { return StaticStruct(); }
+	virtual void ToString(FAnsiStringBuilderBase& Out) const override;
+	virtual bool ShouldReconcile(const FMoverDataStructBase& AuthorityState) const override;
+	virtual void Interpolate(const FMoverDataStructBase& From, const FMoverDataStructBase& To, float Pct) override;
+};
+
+USTRUCT()
+struct FChaosPathedMovementModeDebugData : public FMoverDataStructBase
+{
+	GENERATED_BODY()
+
+	virtual FMoverDataStructBase* Clone() const override;
+	virtual UScriptStruct* GetScriptStruct() const override;
+
+	void Advance()
+	{
+		PreviousServerFrame = ServerFrame;
+		bPlaybackStartedThisFrame = false;
+		bPlaybackStoppedThisFrame = false;
+		bPlaybackBouncedThisFrame = false;
+		bPlaybackReachedEndThisFrame = false;
+		bPlaybackReverseChangedThisFrame = false;
+		bPlaybackLoopingChangedThisFrame = false;
+		bPlaybackOneWayChangedThisFrame = false;
+	}
+
+	// This is the server frame last time we ticked, it might not be ServerFrame-1 if we are resimming
+	UPROPERTY(VisibleAnywhere, Category = "Mover Info")
+	int32 PreviousServerFrame = INDEX_NONE;
+
+	UPROPERTY(VisibleAnywhere, Category = "Mover Info")
+	int32 ServerFrame = INDEX_NONE;
+
+	UPROPERTY(VisibleAnywhere, Category = "Mover Info")
+	FChaosPathedMovementState PreSimMoveState;
+
+	UPROPERTY(VisibleAnywhere, Category = "Mover Info")
+	FChaosPathedMovementState PostSimMoveState;
+
+	UPROPERTY(VisibleAnywhere, Category = "Mover Info")
+	FVector FinalWorldTranslation;
+
+	UPROPERTY(VisibleAnywhere, Category = "Mover Info")
+	FRotator FinalWorldRotation;
+
+	UPROPERTY(VisibleAnywhere, Category = "Mover Info")
+	bool bPlaybackStartedThisFrame = false;
+
+	UPROPERTY(VisibleAnywhere, Category = "Mover Info")
+	bool bPlaybackStoppedThisFrame = false;
+
+	UPROPERTY(VisibleAnywhere, Category = "Mover Info")
+	bool bPlaybackBouncedThisFrame = false;
+
+	UPROPERTY(VisibleAnywhere, Category = "Mover Info")
+	bool bPlaybackReachedEndThisFrame = false;
+
+	UPROPERTY(VisibleAnywhere, Category = "Mover Info")
+	bool bPlaybackReverseChangedThisFrame = false;
+
+	UPROPERTY(VisibleAnywhere, Category = "Mover Info")
+	bool bPlaybackLoopingChangedThisFrame = false;
+
+	UPROPERTY(VisibleAnywhere, Category = "Mover Info")
+	bool bPlaybackOneWayChangedThisFrame = false;
+};
+
+USTRUCT()
+struct FPathedMovementStartedEventData : public FMoverSimulationEventData
+{
+	GENERATED_BODY()
+
+	FPathedMovementStartedEventData(float InEventTimeMs)
+		: FMoverSimulationEventData(InEventTimeMs)
+	{
+	}
+	FPathedMovementStartedEventData() {}
+
+	virtual UScriptStruct* GetScriptStruct() const override
+	{
+		return FPathedMovementStartedEventData::StaticStruct();
+	}
+};
+
+USTRUCT()
+struct FPathedMovementStoppedEventData : public FMoverSimulationEventData
+{
+	GENERATED_BODY()
+
+	FPathedMovementStoppedEventData(float InEventTimeMs, bool InbReachedEndOfPlayback)
+		: FMoverSimulationEventData(InEventTimeMs)
+		, bReachedEndOfPlayback(InbReachedEndOfPlayback)
+	{
+	}
+	FPathedMovementStoppedEventData() {}
+
+	virtual UScriptStruct* GetScriptStruct() const override
+	{
+		return FPathedMovementStoppedEventData::StaticStruct();
+	}
+
+	bool bReachedEndOfPlayback = false;
+};
+
+USTRUCT()
+struct FPathedMovementBouncedEventData : public FMoverSimulationEventData
+{
+	GENERATED_BODY()
+
+	FPathedMovementBouncedEventData(float InEventTimeMs)
+		: FMoverSimulationEventData(InEventTimeMs)
+	{
+	}
+	FPathedMovementBouncedEventData() {}
+
+	virtual UScriptStruct* GetScriptStruct() const override
+	{
+		return FPathedMovementBouncedEventData::StaticStruct();
+	}
+};
+
+USTRUCT()
+struct FPathedMovementInReverseChangedEventData : public FMoverSimulationEventData
+{
+	GENERATED_BODY()
+
+	FPathedMovementInReverseChangedEventData(float InEventTimeMs, bool InbNewInReverse)
+		: FMoverSimulationEventData(InEventTimeMs)
+		, bNewInReverse(InbNewInReverse)
+	{
+	}
+	FPathedMovementInReverseChangedEventData() {}
+
+	virtual UScriptStruct* GetScriptStruct() const override
+	{
+		return FPathedMovementInReverseChangedEventData::StaticStruct();
+	}
+
+	bool bNewInReverse = false;
+};
+
+USTRUCT()
+struct FPathedMovementLoopingChangedEventData : public FMoverSimulationEventData
+{
+	GENERATED_BODY()
+
+	FPathedMovementLoopingChangedEventData(float InEventTimeMs, bool InbNewLooping)
+		: FMoverSimulationEventData(InEventTimeMs)
+		, bNewLooping(InbNewLooping)
+	{
+	}
+	FPathedMovementLoopingChangedEventData() {}
+
+	virtual UScriptStruct* GetScriptStruct() const override
+	{
+		return FPathedMovementLoopingChangedEventData::StaticStruct();
+	}
+
+	bool bNewLooping = false;
+};
+
+USTRUCT()
+struct FPathedMovementOneWayChangedEventData : public FMoverSimulationEventData
+{
+	GENERATED_BODY()
+
+	FPathedMovementOneWayChangedEventData(float InEventTimeMs, bool InbNewOneWay)
+		: FMoverSimulationEventData(InEventTimeMs)
+		, bNewOneWay(InbNewOneWay)
+	{
+	}
+	FPathedMovementOneWayChangedEventData() {}
+
+	virtual UScriptStruct* GetScriptStruct() const override
+	{
+		return FPathedMovementOneWayChangedEventData::StaticStruct();
+	}
+
+	bool bNewOneWay = false;
+};
+
+// Interface for mover modes which can move the controlled component
+// with a general purpose joint constraint or kinematically
+UINTERFACE(MinimalAPI)
+class UChaosMovementActuationInterface : public UInterface
+{
+	GENERATED_BODY()
+};
+class IChaosMovementActuationInterface
+{
+	GENERATED_BODY()
+public:
+
+	// Whether the simulation should use a constraint to move the actor along the path, or moving the actor kinematically
+	virtual bool ShouldUseConstraint() const = 0;
+
+	// Constraint settings of the joint constraint, only valid when ShouldUseConstraint() is true
+	virtual const Chaos::FPBDJointSettings& GetConstraintSettings() const = 0;
+};
+
+// Interface for mover modes moving on a path, with a constraint or kinematically
+// Paths are spatial trajectories with a beginning and an end. A unique position on the path
+// can be calculated by providing a path fraction in [0,1], 0 being at the beginning of the path
+// 1 being at the end.
+UINTERFACE(MinimalAPI)
+class UChaosPathedMovementModeInterface : public UInterface
+{
+	GENERATED_BODY()
+};
+class IChaosPathedMovementModeInterface
+{
+	GENERATED_BODY()
+public:
+
+	// Returns the transform on the path at the path fraction
+	virtual Chaos::FRigidTransform3 CalcTargetRelativeTransform(float PathFraction) const = 0;
+};
diff --git a/Engine/Plugins/Experimental/ChaosMover/Source/ChaosMover/Public/ChaosMover/PathedMovement/ChaosPointMovementPathPattern.h b/Engine/Plugins/Experimental/ChaosMover/Source/ChaosMover/Public/ChaosMover/PathedMovement/ChaosPointMovementPathPattern.h
new file mode 100644
--- /dev/null
+++ b/Engine/Plugins/Experimental/ChaosMover/Source/ChaosMover/Public/ChaosMover/PathedMovement/ChaosPointMovementPathPattern.h
@@ -0,0 +1,84 @@
+// Copyright Epic Games, Inc. All Rights Reserved.
+
+#pragma once
+
+#include "CoreMinimal.h"
+#include "ChaosMover/PathedMovement/ChaosPathedMovementPatternBase.h"
+
+#include "ChaosPointMovementPathPattern.generated.h"
+
+UENUM()
+enum class EChaosPointMovementLocationBasis : uint8
+{
+	/** The location is relative to the previous point in the path */
+	PreviousPoint,
+	/** The location is relative to the origin of the path */
+	PathOrigin,
+	/** The location is relative to the world origin */
+	World
+};
+
+USTRUCT(BlueprintType)
+struct FChaosPointMovementPathPoint
+{
+	GENERATED_BODY()
+
+public:
+	FChaosPointMovementPathPoint() {}
+
+	//@todo DanH: Same options for control as FSplinePathPatternPointData - might want a shared base
+	
+	/** Basis that the location of this point is relative to */
+	UPROPERTY(EditAnywhere, BlueprintReadOnly, Category = Location)
+	EChaosPointMovementLocationBasis Basis = EChaosPointMovementLocationBasis::PreviousPoint;
+	
+	/** The location to move toward, relative to the path origin (i.e. root component location by default) */
+	UPROPERTY(EditAnywhere, BlueprintReadOnly, Category = Location, meta = (MakeEditWidget = true))
+	FVector Location = FVector::Zero();
+
+	/**
+	 * The progress along the path that this point corresponds to.
+	 * Read-only, calculated automatically based on number of points, distances between them, and movement rules for the point 
+	 */
+	// UPROPERTY(VisibleAnywhere, BlueprintReadOnly, Transient, Category = Location)
+	mutable float Progress = 0.f;
+	
+	/** Distance from the start of the path to this point. Read-only. */
+	// UPROPERTY(VisibleAnywhere, BlueprintReadOnly, Transient, Category = Location)
+	mutable float DistanceFromStart = 0.f;
+
+	mutable FVector WorldLoc = FVector::Zero();
+};
+
+/** Movement pattern that moves between explicitly defined points */
+UCLASS()
+class UChaosPointMovementPathPattern : public UChaosPathedMovementPatternBase
+{
+	GENERATED_BODY()
+
+public:
+	virtual void InitializePattern(UChaosMoverSimulation* InSimulation) override;
+
+#if WITH_EDITOR
+	virtual void PostEditChangeProperty(struct FPropertyChangedEvent& PropertyChangedEvent) override;
+#endif
+	
+	virtual bool DebugDrawUsingStepSamples() const override { return false; }
+	//virtual void AppendDebugDrawElements(UChaosPathedPhysicsDebugDrawComponent& DebugDrawComp, FBoxSphereBounds::Builder& InOutDebugBoundsBuilder) override;
+
+protected:
+	virtual FTransform CalcUnmaskedTargetRelativeTransform(float PatternProgress, const FTransform& CurTargetTransform) const override;
+
+	//@todo DanH: We need to force a refresh whenever the path origin changes
+	void RefreshAssignedPointProgress(bool bForceRefresh = false) const;
+
+	/** Relative point locations to move toward, in sequence from first to last */
+	UPROPERTY(EditAnywhere, BlueprintReadWrite, Category = "Point Movement Pattern")
+	TArray<FChaosPointMovementPathPoint> PathPoints;
+
+	/** The sum total distance of this path. Read-only for reference only. */
+	// UPROPERTY(VisibleAnywhere, BlueprintReadOnly, Transient, Category = "Point Movement Pattern")
+	mutable float TotalPathDistance = 0.f;
+
+	mutable bool bHasAssignedPointProgress = false;
+};
\ No newline at end of file
diff --git a/Engine/Plugins/Experimental/ChaosMover/Source/ChaosMover/Public/ChaosMover/PathedMovement/ChaosSplineMovementPathPattern.h b/Engine/Plugins/Experimental/ChaosMover/Source/ChaosMover/Public/ChaosMover/PathedMovement/ChaosSplineMovementPathPattern.h
new file mode 100644
--- /dev/null
+++ b/Engine/Plugins/Experimental/ChaosMover/Source/ChaosMover/Public/ChaosMover/PathedMovement/ChaosSplineMovementPathPattern.h
@@ -0,0 +1,68 @@
+// Copyright Epic Games, Inc. All Rights Reserved.
+
+#pragma once
+
+#include "CoreMinimal.h"
+#include "Engine/EngineTypes.h"
+#include "ChaosMover/PathedMovement/ChaosPathedMovementPatternBase.h"
+
+#include "ChaosSplineMovementPathPattern.generated.h"
+
+class USplineComponent;
+
+/** This isn't formal "spline metadata" because that requires a custom spline and component - we want to add info to vanilla splines */
+USTRUCT()
+struct FChaosSplinePathPatternPointData
+{
+	GENERATED_BODY()
+
+	/** Key of the spline point (or somewhere between two points) where this metadata applies */
+	UPROPERTY()
+	float SplinePointKey = 0.f;
+};
+
+UCLASS()
+class UChaosSplineMovementPathPattern : public UChaosPathedMovementPatternBase
+{
+	GENERATED_BODY()
+	
+public:
+	virtual void InitializePattern(UChaosMoverSimulation* InSimulation) override;
+
+	// The spline is already drawn, no need to draw it again at lower rez
+	virtual bool DebugDrawUsingStepSamples() const override { return false; }
+	
+	bool HasValidBounds() const { return UpperBound > LowerBound; }
+	
+protected:
+	virtual FTransform CalcUnmaskedTargetRelativeTransform(float PatternProgress, const FTransform& CurTargetTransform) const override;
+
+	/** How far into the spline the movement path actually begins.  */
+	UPROPERTY(EditAnywhere, Category = SplinePathPattern, meta = (UIMin = 0.f, ClampMin = 0.f, UIMax = 1.f, ClampMax = 1.f))
+	float LowerBound = 0.f;
+
+	/** How far into the spline the movement path ends. Must be greater than the lower bound. */
+	UPROPERTY(EditAnywhere, Category = SplinePathPattern, meta = (UIMin = 0.f, ClampMin = 0.f, UIMax = 1.f, ClampMax = 1.f))
+	float UpperBound = 1.f;
+
+	/**
+	 * If true, the path duration will be shortened according to how much of the spline is not being followed.
+	 * If false (default), the path duration is unchanged, so the object will move slower when the usable spline range is reduced.
+	 */
+	UPROPERTY(EditAnywhere, Category = SplinePathPattern)
+	bool bOffsetsModifyDuration = false;
+
+	UPROPERTY(EditAnywhere, Category = SplinePathPattern)
+	bool bApplySplineScaling = false;
+	
+	/**
+	 * Optional property to specify the spline component that defines the path to follow. If blank, we'll use the first spline component we find in this actor.
+	 * This is only necessary to set if you have multiple spline components on the actor, or want to follow a spline on an external actor.
+	 * Thread Safety Warning, this spline cannot be modified whilst in use
+	 */
+	UPROPERTY(EditAnywhere, DisplayName = "Spline Component", Category = SplinePathPattern, meta = (AllowAnyActor, UseComponentPicker, AllowedClasses = "/Script/Engine.SplineComponent"))
+	FComponentReference SplineComponentRef;
+
+	UPROPERTY(Transient)
+	TObjectPtr<const USplineComponent> SplineComp;
+};
diff --git a/Engine/Plugins/Experimental/ChaosMover/Source/ChaosMover/Public/ChaosMover/PathedMovement/Transitions/ChaosPathedMovementReachedEndTransition.h b/Engine/Plugins/Experimental/ChaosMover/Source/ChaosMover/Public/ChaosMover/PathedMovement/Transitions/ChaosPathedMovementReachedEndTransition.h
new file mode 100644
--- /dev/null
+++ b/Engine/Plugins/Experimental/ChaosMover/Source/ChaosMover/Public/ChaosMover/PathedMovement/Transitions/ChaosPathedMovementReachedEndTransition.h
@@ -0,0 +1,24 @@
+// Copyright Epic Games, Inc. All Rights Reserved.
+
+#pragma once
+
+#include "ChaosMover/ChaosMovementModeTransition.h"
+
+#include "ChaosPathedMovementReachedEndTransition.generated.h"
+
+// Transition checking if an ongiong pathed movement mode has reached the end of the path
+UCLASS(Blueprintable, BlueprintType, EditInlineNew, DefaultToInstanced)
+class CHAOSMOVER_API UChaosPathedMovementReachedEndTransition : public UChaosMovementModeTransition
+{
+	GENERATED_BODY()
+
+public:
+	UChaosPathedMovementReachedEndTransition(const FObjectInitializer& ObjectInitializer);
+
+	virtual FTransitionEvalResult Evaluate_Implementation(const FSimulationTickParams& Params) const override;
+	virtual void Trigger_Implementation(const FSimulationTickParams& Params) override;
+
+	/** Name of movement mode to transition to when the pathed movement has reached the end of the path */
+	UPROPERTY(BlueprintReadWrite, EditAnywhere, Category = Parameters)
+	FName TransitionToMode;
+};
\ No newline at end of file
diff --git a/Engine/Plugins/Experimental/Mover/Source/Mover/Private/MoverComponent.cpp b/Engine/Plugins/Experimental/Mover/Source/Mover/Private/MoverComponent.cpp
--- a/Engine/Plugins/Experimental/Mover/Source/Mover/Private/MoverComponent.cpp
+++ b/Engine/Plugins/Experimental/Mover/Source/Mover/Private/MoverComponent.cpp
@@ -242,7 +242,8 @@ void UMoverComponent::BeginPlay()
 	WorldToGravityTransform = FQuat::FindBetweenNormals(FVector::UpVector, GetUpDirection());
 	GravityToWorldTransform = WorldToGravityTransform.Inverse();
 	
-	if (const AActor* MyActor = GetOwner())
+	AActor* MyActor = GetOwner();
+	if (MyActor)
 	{
 		// If no primary visual component is already set, fall back to searching for any kind of mesh,
 		// favoring a direct scene child of the UpdatedComponent.
@@ -274,20 +275,30 @@ void UMoverComponent::BeginPlay()
 			UMotionWarpingMoverAdapter* WarpingAdapter = WarpingComp->CreateOwnerAdapter<UMotionWarpingMoverAdapter>();
 			WarpingAdapter->SetMoverComp(this);
 		}
-	}
 
-	// If an InputProducer isn't already set, check the actor and its components for one
-	if (!InputProducer)
-	{
-		if (AActor* ActorOwner = GetOwner())
+		// If an InputProducer isn't already set, check if the actor is one
+		if (!InputProducer)
 		{
-			if (ActorOwner->GetClass()->ImplementsInterface(UMoverInputProducerInterface::StaticClass()))
+			if (MyActor->GetClass()->ImplementsInterface(UMoverInputProducerInterface::StaticClass()))
 			{
-				InputProducer = ActorOwner;
+				InputProducer = MyActor;
+				InputProducers.Add(MyActor);
 			}
-			else if (UActorComponent* FoundInputProducerComp = ActorOwner->FindComponentByInterface(UMoverInputProducerInterface::StaticClass()))
+		}
+
+		TSet<UActorComponent*> Components = MyActor->GetComponents();
+		for (UActorComponent* Component : Components)
+		{
+			if (IsValid(Component))
 			{
-				InputProducer = FoundInputProducerComp;
+				if (Component->GetClass()->ImplementsInterface(UMoverInputProducerInterface::StaticClass()))
+				{
+					// Let's not add that component if a derived class made it the InputProducer
+					if (InputProducer != Component)
+					{
+						InputProducers.Add(Component);
+					}
+				}
 			}
 		}
 	}
@@ -312,10 +323,12 @@ void UMoverComponent::ProduceInput(const int32 DeltaTimeMS, FMoverInputCmdContex
 {
 	Cmd->InputCollection.Empty();
 	
-	// Pass input production on to the right party
-	if (InputProducer)
+	for (TObjectPtr<UObject> InputProducerComponent : InputProducers)
 	{
-		IMoverInputProducerInterface::Execute_ProduceInput(InputProducer, DeltaTimeMS, OUT *Cmd);
+		if (IsValid(InputProducerComponent))
+		{
+			IMoverInputProducerInterface::Execute_ProduceInput(InputProducerComponent, DeltaTimeMS, IN OUT *Cmd);
+		}
 	}
 
 	CachedLastProducedInputCmd = *Cmd;
@@ -1599,6 +1612,11 @@ const FLayeredMoveBase* UMoverComponent::FindActiveLayeredMoveByType(const UScri
 }
 
 void UMoverComponent::QueueNextMode(FName DesiredModeName, bool bShouldReenter)
+{
+	DoQueueNextMode(DesiredModeName, bShouldReenter);
+}
+
+void UMoverComponent::DoQueueNextMode(FName DesiredModeName, bool bShouldReenter)
 {
 	ModeFSM->QueueNextMode(DesiredModeName, bShouldReenter);
 }
@@ -2233,10 +2251,18 @@ void UMoverComponent::SetSimulationOutput(const FMoverTimeStep& TimeStep, const
 	{
 		if (const FMoverSimulationEventData* Data = EventData.Get())
 		{
+			// Process the simulation event at the mover component (or derived) level
 			ProcessSimulationEvent(*Data);
+
+			// Broadcast the event to 
+			if (OnPostSimEventReceived.IsBound())
+			{
+				OnPostSimEventReceived.Broadcast(*Data);
+			}
 		}
 	}
-	// This is for things like the ground info that we want to cache and interpolate but isn't part of the networked sync state
+
+	// This is for things like the ground info that we want to cache and interpolate but isn't part of the networked sync state.
 	// AdditionalOutputData is generic because ground info might not be useful for platforms, say, but platforms may want to store something else.
 	SetAdditionalSimulationOutput(OutputData.AdditionalOutputData);
 }
diff --git a/Engine/Plugins/Experimental/Mover/Source/Mover/Public/MoverComponent.h b/Engine/Plugins/Experimental/Mover/Source/Mover/Public/MoverComponent.h
--- a/Engine/Plugins/Experimental/Mover/Source/Mover/Public/MoverComponent.h
+++ b/Engine/Plugins/Experimental/Mover/Source/Mover/Public/MoverComponent.h
@@ -18,10 +18,12 @@
 #include "MovementModifier.h"
 #include "Backends/MoverBackendLiaison.h"
 #include "UObject/WeakInterfacePtr.h"
+#include "Templates/SharedPointer.h"
 #include "MoverComponent.generated.h"
 
 struct FMoverTimeStep;
 struct FInstantMovementEffect;
+struct FMoverSimulationEventData;
 class UMovementModeStateMachine;
 class UMovementMixer;
 
@@ -55,6 +57,10 @@ DECLARE_DYNAMIC_MULTICAST_DELEGATE_TwoParams(FMover_OnPostFinalize, const FMover
 // Fired after proposed movement has been generated (i.e. after movement modes and layered moves have generated movement and mixed together).
 DECLARE_DYNAMIC_DELEGATE_ThreeParams(FMover_ProcessGeneratedMovement, const FMoverTickStartData&, StartState, const FMoverTimeStep&, TimeStep, FProposedMove&, OutProposedMove);
 
+// Fired when a new event has been received from the simulation. This is a C++ only dispatch of the generic event. To dispatch the event to BP, prefer exposing the concrete event
+// via a dedicated dynamic delegate (like OnMovementModeChanged).
+DECLARE_MULTICAST_DELEGATE_OneParam(FMover_OnPostSimEventReceived, const FMoverSimulationEventData&);
+
 /**
  * 
  */
@@ -99,6 +105,11 @@ public:
 	UPROPERTY(BlueprintAssignable, Category = Mover)
 	FMover_OnPostFinalize OnPostFinalize;
 
+	// Fired when a new event has been received from the simulation
+	// This happens after the event had been processed at the mover component level, which may
+	// have caused a dedicated delegate, e.g. OnMovementModeChanged, to fire prior to this broadcast.
+	FMover_OnPostSimEventReceived OnPostSimEventReceived;
+
 	/**
 	 * Broadcast after proposed movement has been generated. After movement modes and layered moves have generated movement and mixed together.
 	 * This allows for final modifications to proposed movement before it's executed.
@@ -163,10 +174,27 @@ public:
 	UPROPERTY(BlueprintReadWrite, EditAnywhere, Category=Mover)
 	TArray<FMoverDataPersistence> PersistentSyncStateDataTypes;
 
-	/** Optional object for producing input cmds. Typically set at BeginPlay time. If not specified, defaulted input will be used. */
+	/** Optional object for producing input cmds. Typically set at BeginPlay time. If not specified, defaulted input will be used.
+	*   Note that any other actor component implementing MoverInputProducerInterface on this component's owner will also be able
+	*   to produce input commands if bGatherInputFromAllInputProducerComponents is true. @see bGatherInputFromAllInputProducerComponents
+	*/
 	UPROPERTY(BlueprintReadWrite, EditAnywhere, Category = Mover, meta = (ObjectMustImplement = "/Script/Mover.MoverInputProducerInterface"))
 	TObjectPtr<UObject> InputProducer;
 
+	/** If true, any actor component implementing MoverInputProducerInterface on this component's owner will be able to produce input commands */
+	UPROPERTY(EditAnywhere, BlueprintReadOnly, Category=Mover)
+	bool bGatherInputFromAllInputProducerComponents = true;
+
+	/* All MoverInputProducerInterface objects producing input for this mover component. If bGatherInputFromAllInputProducerComponents
+	*  is true, all components implementing MoverInputProducerInterface on this component's owner will be added to 
+	*  this array at BeginPlay time, and IMoverInputProducerInterface::ProduceInput will be called on each within UMoverComponent::ProduceInput.
+	*  The order shouldn't matter, as this is for input commands independent of each other, driving different movement modes.
+	*  If order is important, set bGatherInputFromAllInputProducerComponents to false and implement a dedicated input component instead,
+	*  gathering input from different sources in a custom order and set it as the InputProducer.
+	*/
+	UPROPERTY(Transient)
+	TArray<TObjectPtr<UObject>> InputProducers;
+
 	/** Optional object for mixing proposed moves.Typically set at BeginPlay time. If not specified, UDefaultMovementMixer will be used. */
 	UPROPERTY(BlueprintReadWrite, EditAnywhere, Category = Mover)
 	TObjectPtr<UMovementMixer> MovementMixer;
@@ -543,6 +571,8 @@ protected:
 	MOVER_API void FindDefaultUpdatedComponent();
 	MOVER_API void UpdateTickRegistration();
 
+	MOVER_API virtual void DoQueueNextMode(FName DesiredModeName, bool bShouldReenter=false);
+
 	// Broadcast during the simulation tick after inputs have been processed, but before the actual move is performed.
 	// Note - When async simulating, the delegate would be called on the async thread, and might be broadcast multiple times.
 	UPROPERTY(BlueprintAssignable, Category = Mover)
diff --git a/Engine/Plugins/Experimental/Mover/Source/Mover/Public/MoverSimulationTypes.h b/Engine/Plugins/Experimental/Mover/Source/Mover/Public/MoverSimulationTypes.h
--- a/Engine/Plugins/Experimental/Mover/Source/Mover/Public/MoverSimulationTypes.h
+++ b/Engine/Plugins/Experimental/Mover/Source/Mover/Public/MoverSimulationTypes.h
@@ -365,7 +365,7 @@ class IMoverInputProducerInterface : public IInterface
 public:
 	/** Contributes additions to the input cmd for this simulation frame. Typically this is translating accumulated user input (or AI state) into parameters that affect movement. */
 	UFUNCTION(BlueprintNativeEvent)
-	void ProduceInput(int32 SimTimeMs, FMoverInputCmdContext& InputCmdResult);
+	MOVER_API void ProduceInput(int32 SimTimeMs, FMoverInputCmdContext& InputCmdResult);
 };
 
 
