Build fix: Static Analysis spotted an error in Interpolate code, it was blending between To and To instead of between From and To

diff --git a/Engine/Plugins/Experimental/ChaosMover/Source/ChaosMover/Private/PathedMovement/ChaosPathedMovementTypes.cpp b/Engine/Plugins/Experimental/ChaosMover/Source/ChaosMover/Private/PathedMovement/ChaosPathedMovementTypes.cpp
--- a/Engine/Plugins/Experimental/ChaosMover/Source/ChaosMover/Private/PathedMovement/ChaosPathedMovementTypes.cpp
+++ b/Engine/Plugins/Experimental/ChaosMover/Source/ChaosMover/Private/PathedMovement/ChaosPathedMovementTypes.cpp
@@ -144,8 +144,8 @@ void FChaosPathedMovementState::Interpolate(const FMoverDataStructBase& From, co
 	const FChaosPathedMovementState& TypedFrom = static_cast<const FChaosPathedMovementState&>(From);
 	const FChaosPathedMovementState& TypedTo = static_cast<const FChaosPathedMovementState&>(To);
 
-	PropertiesInEffect = (Pct < .5f) ? TypedTo.PropertiesInEffect : TypedTo.PropertiesInEffect;
-	LatestMovementStartFrame = (Pct < .5f) ? TypedTo.LatestMovementStartFrame : TypedTo.LatestMovementStartFrame;
+	PropertiesInEffect = (Pct < .5f) ? TypedFrom.PropertiesInEffect : TypedTo.PropertiesInEffect;
+	LatestMovementStartFrame = (Pct < .5f) ? TypedFrom.LatestMovementStartFrame : TypedTo.LatestMovementStartFrame;
 	CurrentPathFraction = FMath::Lerp(TypedFrom.CurrentPathFraction, TypedTo.CurrentPathFraction, Pct);
 	bIsPathFractionIncreasing = (Pct < .5f) ? TypedFrom.bIsPathFractionIncreasing : TypedTo.bIsPathFractionIncreasing;
 	bHasFinished = (Pct < .5f) ? TypedFrom.bHasFinished : TypedTo.bHasFinished;
