Build fix: Adding missing World.h include

diff --git a/Engine/Plugins/Experimental/ChaosMover/Source/ChaosMover/Private/PathedMovement/ChaosPathedMovementControllerComponent.cpp b/Engine/Plugins/Experimental/ChaosMover/Source/ChaosMover/Private/PathedMovement/ChaosPathedMovementControllerComponent.cpp
--- a/Engine/Plugins/Experimental/ChaosMover/Source/ChaosMover/Private/PathedMovement/ChaosPathedMovementControllerComponent.cpp
+++ b/Engine/Plugins/Experimental/ChaosMover/Source/ChaosMover/Private/PathedMovement/ChaosPathedMovementControllerComponent.cpp
@@ -7,6 +7,7 @@
 #include "Physics/Experimental/PhysScene_Chaos.h"
 #include "MoverComponent.h"
 #include "TimerManager.h"
+#include "Engine/World.h"
 
 #include UE_INLINE_GENERATED_CPP_BY_NAME(ChaosPathedMovementControllerComponent)
 
