Removed call to deprecated water Query function.

diff --git a/Engine/Plugins/Experimental/ChaosMover/Source/ChaosMover/Private/Utilities/ChaosMoverQueryUtils.cpp b/Engine/Plugins/Experimental/ChaosMover/Source/ChaosMover/Private/Utilities/ChaosMoverQueryUtils.cpp
--- a/Engine/Plugins/Experimental/ChaosMover/Source/ChaosMover/Private/Utilities/ChaosMoverQueryUtils.cpp
+++ b/Engine/Plugins/Experimental/ChaosMover/Source/ChaosMover/Private/Utilities/ChaosMoverQueryUtils.cpp
@@ -42,26 +42,34 @@ namespace UE::ChaosMover::Utils
 
 							OutWaterResult.WaterSplineData.WaterBody = WaterBody;
 
-							FWaterBodyQueryResult QueryResult = WaterBody->GetWaterBodyComponent()->QueryWaterInfoClosestToWorldLocation(
+							TValueOrError<FWaterBodyQueryResult, EWaterBodyQueryError> QueryResult = WaterBody->GetWaterBodyComponent()->TryQueryWaterInfoClosestToWorldLocation(
 								Location,
 								EWaterBodyQueryFlags::ComputeLocation | EWaterBodyQueryFlags::ComputeNormal | EWaterBodyQueryFlags::ComputeImmersionDepth,
 								OutWaterResult.WaterSplineData.SplineInputKey
 							);
 
-							OutWaterResult.WaterSplineData.ImmersionDepth = QueryResult.GetImmersionDepth();
+							// If the query has no value, treat it as if this was never in the hit list:
+							if (!QueryResult.HasValue())
+							{
+								continue;
+							}
 
-							OutWaterResult.WaterSplineData.WaterPlaneLocation = QueryResult.GetWaterPlaneLocation();
-							OutWaterResult.WaterSplineData.WaterPlaneNormal = QueryResult.GetWaterPlaneNormal();
+							const FWaterBodyQueryResult& Query = QueryResult.GetValue();
+
+							OutWaterResult.WaterSplineData.ImmersionDepth = Query.GetImmersionDepth();
+
+							OutWaterResult.WaterSplineData.WaterPlaneLocation = Query.GetWaterPlaneLocation();
+							OutWaterResult.WaterSplineData.WaterPlaneNormal = Query.GetWaterPlaneNormal();
 
 							const float CapsuleBottom = Location.Z;
 							const float CapsuleTop = Location.Z + (TargetHeight * 2);
-							OutWaterResult.WaterSplineData.WaterSurfaceLocation = QueryResult.GetWaterSurfaceLocation();
+							OutWaterResult.WaterSplineData.WaterSurfaceLocation = Query.GetWaterSurfaceLocation();
 
 							OutWaterResult.WaterSplineData.WaterSurfaceOffset = OutWaterResult.WaterSplineData.WaterSurfaceLocation - Location;
 
 							OutWaterResult.WaterSplineData.ImmersionPercent = FMath::Clamp((OutWaterResult.WaterSplineData.WaterSurfaceLocation.Z - CapsuleBottom) / (CapsuleTop - CapsuleBottom), 0.f, 1.f);
 
-							OutWaterResult.WaterSplineData.WaterSurfaceNormal = QueryResult.GetWaterSurfaceNormal();
+							OutWaterResult.WaterSplineData.WaterSurfaceNormal = Query.GetWaterSurfaceNormal();
 
 	#if CHAOSMOVER_DEBUG_DRAW
 							Chaos::FDebugDrawQueue::GetInstance().DrawDebugLine(Location, Location - FVector::UpVector * OutWaterResult.WaterSplineData.ImmersionDepth, FColor::Blue, false, -1.f, 10, 1.0f);
