ChaosMover: Porting Dan's Pathed Movement debug draws to ChaosMover

diff --git a/Engine/Plugins/Experimental/ChaosMover/Source/ChaosMover/Private/ChaosMoverSimulation.cpp b/Engine/Plugins/Experimental/ChaosMover/Source/ChaosMover/Private/ChaosMoverSimulation.cpp
--- a/Engine/Plugins/Experimental/ChaosMover/Source/ChaosMover/Private/ChaosMoverSimulation.cpp
+++ b/Engine/Plugins/Experimental/ChaosMover/Source/ChaosMover/Private/ChaosMoverSimulation.cpp
@@ -104,13 +104,8 @@ void UChaosMoverSimulation::Init(const FInitParams& InitParams)
 
 	CurrentSyncState = InitParams.InitialSyncState;
 
-	UE::ChaosMover::FMoverStateMachine::FInitParams StateMachineInitParams;
-	StateMachineInitParams.ImmediateMovementModeTransition = InitParams.ImmediateModeTransition;
-	StateMachineInitParams.NullMovementMode = InitParams.NullMovementMode;
-	StateMachineInitParams.Simulation = this;
-	StateMachine.Init(StateMachineInitParams);
-
 	// Set the pathed movement basis to the initial location of the controlled particle
+	// We set this up before state machine init so the pathed movement modes have a pathed movement basis to work with on registration
 	if (PhysicsObject)
 	{
 		if (Chaos::FGeometryParticle* Particle = FPhysicsObjectExternalInterface::LockRead(PhysicsObject)->GetRigidParticle(PhysicsObject))
@@ -122,6 +117,12 @@ void UChaosMoverSimulation::Init(const FInitParams& InitParams)
 		}
 	}
 
+	UE::ChaosMover::FMoverStateMachine::FInitParams StateMachineInitParams;
+	StateMachineInitParams.ImmediateMovementModeTransition = InitParams.ImmediateModeTransition;
+	StateMachineInitParams.NullMovementMode = InitParams.NullMovementMode;
+	StateMachineInitParams.Simulation = this;
+	StateMachine.Init(StateMachineInitParams);
+
 	for (const TPair<FName, TWeakObjectPtr<UBaseMovementMode>>& Element : InitParams.ModesToRegister)
 	{
 		const FName& ModeName = Element.Key;
@@ -709,15 +710,15 @@ void UChaosMoverSimulation::OnMovementModeChanged(const FMovementModeChangedEven
 					ConstraintHandle->SetSettings(NextMovementActuationInterface->GetConstraintSettings());
 
 					// If the mode was also a pathed mode, let's teleport to the position on the path
-					// we're supposed to be at, using CurrentPathFraction
+					// we're supposed to be at, using CurrentPathProgress
 					Chaos::FRigidTransform3 TargetLastFrame = Chaos::FRigidTransform3::Identity;
 					IChaosPathedMovementModeInterface* NextPathedMovementMode = Cast<IChaosPathedMovementModeInterface>(NextMovementActuationInterface);
 					if (NextPathedMovementMode)
 					{	
-						float PathFraction = PresimPathedMovementState ? PresimPathedMovementState->CurrentPathFraction : 0.0f;
+						const float PathProgress = PresimPathedMovementState ? PresimPathedMovementState->CurrentPathProgress : 0.0f;
 						// We resume the path at the reset progress
-						// (if not 0, CurrentPathFraction was either not reset properly when the movement mode changed, or deliberately set to a different value)
-						TargetLastFrame = NextPathedMovementMode->CalcTargetRelativeTransform(PathFraction);
+						// (if not 0, CurrentPathProgress was either not reset properly when the movement mode changed, or deliberately set to a different value)
+						TargetLastFrame = NextPathedMovementMode->CalcTargetRelativeTransform(PathProgress);
 					}
 					Chaos::FRigidTransform3 WorldTransform = TargetLastFrame * PathBasis;
 					TeleportActuationTarget(WorldTransform);
diff --git a/Engine/Plugins/Experimental/ChaosMover/Source/ChaosMover/Private/PathedMovement/ChaosPathedMovementControllerComponent.cpp b/Engine/Plugins/Experimental/ChaosMover/Source/ChaosMover/Private/PathedMovement/ChaosPathedMovementControllerComponent.cpp
--- a/Engine/Plugins/Experimental/ChaosMover/Source/ChaosMover/Private/PathedMovement/ChaosPathedMovementControllerComponent.cpp
+++ b/Engine/Plugins/Experimental/ChaosMover/Source/ChaosMover/Private/PathedMovement/ChaosPathedMovementControllerComponent.cpp
@@ -1,6 +1,8 @@
 // Copyright Epic Games, Inc. All Rights Reserved.
 
 #include "ChaosMover/PathedMovement/ChaosPathedMovementControllerComponent.h"
+#include "ChaosMover/PathedMovement/ChaosPathedMovementDebugDrawComponent.h"
+#include "ChaosMover/PathedMovement/ChaosPathedMovementMode.h"
 #include "ChaosMover/ChaosMoverLog.h"
 #include "ChaosMover/ChaosMoverSettings.h"
 #include "PBDRigidsSolver.h"
@@ -13,7 +15,43 @@
 
 UChaosPathedMovementControllerComponent::UChaosPathedMovementControllerComponent()
 {
-	
+#if ENABLE_DRAW_DEBUG
+	DebugDrawComp = CreateDefaultSubobject<UChaosPathedMovementDebugDrawComponent>(TEXT("ChaosPathedMovementDebugDraw"));
+
+	// Since the debug draw comp relies on data in this component, and this component will often need to be re-registered following the construction script,
+	// we want registration of the debug draw comp to match (and follow) the registration state of this component
+	DebugDrawComp->bAutoRegister = false;
+#endif
+}
+
+void UChaosPathedMovementControllerComponent::OnRegister()
+{
+	Super::OnRegister();
+
+	AActor* Owner = GetOwner();
+	if (UMoverComponent* OwnerMoverComponent = Owner ? Owner->FindComponentByClass<UMoverComponent>() : nullptr)
+	{
+		for (TPair<FName, TObjectPtr<UBaseMovementMode>>& NameModePair : OwnerMoverComponent->MovementModes)
+		{
+			if (UChaosPathedMovementMode* PathedMovementMode = Cast<UChaosPathedMovementMode>(NameModePair.Value))
+			{
+				PathedMovementMode->InitializePath();
+			}
+		}	
+	}
+
+#if ENABLE_DRAW_DEBUG
+	DebugDrawComp->RegisterComponent();
+#endif
+}
+
+void UChaosPathedMovementControllerComponent::OnUnregister()
+{
+	Super::OnUnregister();
+
+#if ENABLE_DRAW_DEBUG
+	DebugDrawComp->UnregisterComponent();
+#endif
 }
 
 void UChaosPathedMovementControllerComponent::BeginPlay()
@@ -315,4 +353,19 @@ void UChaosPathedMovementControllerComponent::ExecuteOrScheduleChange(const FCha
 			LambdaChangeMutableProperties(NewProperties, Solver->GetCurrentFrame(), bIsStart);
 		}), MinDelaySeconds, false);
 	}
-}
\ No newline at end of file
+}
+
+bool UChaosPathedMovementControllerComponent::ShouldDisplayProgressPreviewMesh_Implementation() const
+{
+	return bDisplayProgressPreviewMesh;
+}
+
+float UChaosPathedMovementControllerComponent::GetPreviewMeshOverallPathProgress_Implementation() const
+{
+	return PreviewMeshOverallPathProgress;
+}
+
+UMaterialInterface* UChaosPathedMovementControllerComponent::GetProgressPreviewMeshMaterial_Implementation() const
+{
+	return PreviewMeshMaterial;
+}
diff --git a/Engine/Plugins/Experimental/ChaosMover/Source/ChaosMover/Private/PathedMovement/ChaosPathedMovementDebugDrawComponent.cpp b/Engine/Plugins/Experimental/ChaosMover/Source/ChaosMover/Private/PathedMovement/ChaosPathedMovementDebugDrawComponent.cpp
new file mode 100644
--- /dev/null
+++ b/Engine/Plugins/Experimental/ChaosMover/Source/ChaosMover/Private/PathedMovement/ChaosPathedMovementDebugDrawComponent.cpp
@@ -0,0 +1,404 @@
+// Copyright Epic Games, Inc. All Rights Reserved.
+
+#include "ChaosMover/PathedMovement/ChaosPathedMovementDebugDrawComponent.h"
+
+#include "StaticMeshSceneProxy.h"
+#include "Components/MeshComponent.h"
+#include "Components/StaticMeshComponent.h"
+#include "Engine/World.h"
+#include "ChaosMover/PathedMovement/ChaosPathedMovementMode.h"
+#include "ChaosMover/PathedMovement/ChaosPathedMovementPatternBase.h"
+#include "MoverComponent.h"
+#include "SceneInterface.h"
+
+TAutoConsoleVariable<int32> CVarChaosPathedMovementTotalDebugDrawSteps(
+	TEXT("ChaosMover.PathedMovement.DebugDraw.TotalNumSteps"),
+	64,
+	TEXT("How many steps/lines to draw for each debug drawn path"));
+	
+TAutoConsoleVariable<int32> CVarChaosPathedMovementDisplayedDebugDrawSteps(
+	TEXT("ChaosMover.PathedMovement.DebugDraw.DisplayedSteps"),
+	0,
+	TEXT("Of the total number of steps in the path debug draw (see ChaosMover.PathedMovement.DebugDraw.TotalNumSteps), how many should actually get drawn? If <= 0, all steps are drawn."));
+
+class FChaosPathedMovementDebugRenderSceneProxy : public FDebugRenderSceneProxy
+{
+public:
+	explicit FChaosPathedMovementDebugRenderSceneProxy(const UPrimitiveComponent& InComponent)
+		: FDebugRenderSceneProxy(&InComponent)
+	{}
+
+	virtual SIZE_T GetTypeHash() const override 
+	{
+		static size_t UniquePointer;
+		return reinterpret_cast<size_t>(&UniquePointer);
+	}
+	
+	virtual FPrimitiveViewRelevance GetViewRelevance(const FSceneView* View) const override
+	{
+		UWorld* World = (View && View->Family && View->Family->Scene) ? View->Family->Scene->GetWorld() : nullptr;
+		bool bIsAssetViewport = World && (World->WorldType == EWorldType::EditorPreview);
+		FPrimitiveViewRelevance Result;
+		//@todo DanH: Add a show flag for these, like the one for spline components: View->Family->EngineShowFlags.Splines
+		Result.bDrawRelevance = IsShown(View) && (bIsAssetViewport || IsSelected());
+		Result.bDynamicRelevance = true;
+		// Result.bShadowRelevance = IsShadowCast(View);
+		Result.bEditorPrimitiveRelevance = true; // UseEditorCompositing(View);
+		return Result;
+	}
+
+	virtual uint32 GetMemoryFootprint(void) const override { return sizeof(*this) + GetAllocatedSize(); }
+};
+
+UChaosPathedMovementDebugDrawComponent::UChaosPathedMovementDebugDrawComponent()
+{
+	// bIsEditorOnly = true;
+}
+
+void UChaosPathedMovementDebugDrawComponent::OnRegister()
+{
+	// DestroyProgressPreviewMeshComp();
+	UpdatePreviewMeshComp();
+
+#if WITH_EDITOR
+	FCoreUObjectDelegates::OnObjectPropertyChanged.AddUObject(this, &ThisClass::HandleObjectPropertyChanged);
+#endif
+
+	Super::OnRegister();
+}
+
+void UChaosPathedMovementDebugDrawComponent::OnUnregister()
+{
+	Super::OnUnregister();
+
+#if WITH_EDITOR
+	FCoreUObjectDelegates::OnObjectPropertyChanged.RemoveAll(this);
+#endif
+
+	if (ProgressPreviewMeshComp)
+	{
+		ProgressPreviewMeshComp->UnregisterComponent();
+	}
+}
+
+FBoxSphereBounds UChaosPathedMovementDebugDrawComponent::CalcBounds(const FTransform& LocalToWorld) const
+{
+	return Bounds;
+}
+
+void UChaosPathedMovementDebugDrawComponent::CreateRenderState_Concurrent(FRegisterComponentContext* Context)
+{
+	// Calculate everything for the scene proxy here so that our bounds are ready early enough (they get checked just before the scene proxy gets added)
+	FBoxSphereBounds::Builder DebugBoundsBuilder;
+	DebugBoundsBuilder += GetOwner()->GetTransform().GetLocation();
+	
+	DebugLines.Reset();
+	DebugDashedLines.Reset();
+	DebugArrowLines.Reset();
+	DebugSpheres.Reset();
+	DebugStars.Reset();
+
+	const UMoverComponent* MoverComp = GetMoverComponent();
+	if (!MoverComp)
+	{
+		return;
+	}
+
+	const TArray<UActorComponent*> DebugDrawInterfaceComponents = GetOwner()->GetComponentsByInterface(UChaosPathedMovementDebugDrawInterface::StaticClass());
+	if (!DebugDrawInterfaceComponents.IsEmpty())
+	{
+		DebugDrawInterfaceObject = DebugDrawInterfaceComponents[0];
+	}
+	
+	const UChaosPathedMovementMode* PathedMovementMode = MoverComp->FindMode_Mutable<UChaosPathedMovementMode>(MoverComp->StartingMovementMode);
+	if (PathedMovementMode && !PathedMovementMode->PathPatterns.IsEmpty() && (PathedMovementMode->bDebugDrawAggregatePath || PathedMovementMode->bAllowPatternDebugDrawing))
+	{
+		// Collect valid patterns and count the number of total curves to draw to reserve enough space for DebugLines
+		TArray<const UChaosPathedMovementPatternBase*> ValidPatterns;
+		int32 NumCurvesToDraw = PathedMovementMode->bDebugDrawAggregatePath ? 1 : 0;
+		for (UChaosPathedMovementPatternBase* PathPattern : PathedMovementMode->PathPatterns)
+		{
+			if (PathPattern)
+			{
+				ValidPatterns.Emplace(PathPattern);
+				if (PathPattern->DebugDrawUsingStepSamples())
+				{
+					++NumCurvesToDraw;
+				}
+
+				if (PathedMovementMode->bAllowPatternDebugDrawing)
+				{
+					PathPattern->AppendDebugDrawElements(*this, DebugBoundsBuilder);
+				}
+			}
+		}
+
+		// Reserve enough debug lines to not have to reallocate on addition
+		const float NumSteps = CVarChaosPathedMovementTotalDebugDrawSteps.GetValueOnGameThread();
+		DebugLines.Reserve(DebugLines.Num() + NumSteps * NumCurvesToDraw);
+
+		// This is an array of path progress samples that we will use if we draw the aggregate path
+		// The spacing of samples needs to adapt to path patterns or the sampling resolution my
+		// skip entire patterns by stepping too coarsely. As we sample patterns, we may add the same sample multiple
+		// times or out of order. This array will then be sorted and values that are too close to the previous ones
+		// will be skipped
+		TArray<float> PathProgressSamples;
+		if (PathedMovementMode->bDebugDrawAggregatePath)
+		{
+			// We preserve for a minimum but may have to reallocate depending on patterns
+			PathProgressSamples.Reserve(NumSteps);
+		}
+		
+		// This is the inferred path basis
+		const FTransform LocalToWorld = GetOwner()->GetTransform();
+
+		// Establish a previous transform for each path pattern, and add a sample for the aggreagte path for the start of each pattern
+		TArray<FTransform> PreviousStepPatternTransforms;
+		PreviousStepPatternTransforms.SetNumUninitialized(ValidPatterns.Num());
+		for (int32 PatternIdx = 0; PatternIdx < ValidPatterns.Num(); ++PatternIdx)
+		{
+			const UChaosPathedMovementPatternBase& Pattern = *ValidPatterns[PatternIdx];
+			if (PathedMovementMode->bDebugDrawAggregatePath)
+			{
+				// Add an aggregate path sample for the beginning of this path.
+				PathProgressSamples.Add(Pattern.ConvertPatternToPathProgress(Pattern.StartAtPathProgress));
+			}
+			const bool bDrawSampledPattern = Pattern.DebugDrawUsingStepSamples() && PathedMovementMode->bAllowPatternDebugDrawing;
+			if (bDrawSampledPattern)
+			{
+				const FTransform FirstPatternTransform = Chaos::FRigidTransform3::MultiplyNoScale(Pattern.CalcMaskedTargetRelativeTransform(0.0f, LocalToWorld), LocalToWorld);
+				PreviousStepPatternTransforms[PatternIdx] = FirstPatternTransform;
+				DebugBoundsBuilder += FirstPatternTransform.GetLocation();
+			}
+		}			
+		
+		// Now sample each pattern to draw them, and add that sample to the aggregate path sample if we draw it
+		const int32 NumStepsToDraw = CVarChaosPathedMovementDisplayedDebugDrawSteps.GetValueOnGameThread() > 0 ? CVarChaosPathedMovementDisplayedDebugDrawSteps.GetValueOnGameThread() : NumSteps;
+		for (int32 Step = 1; Step <= NumStepsToDraw; ++Step)
+		{
+			for (int32 PatternIdx = 0; PatternIdx < ValidPatterns.Num(); ++PatternIdx)
+			{
+				const UChaosPathedMovementPatternBase& Pattern = *ValidPatterns[PatternIdx];
+				const float PatternProgress = Step / NumSteps;
+				if (PathedMovementMode->bDebugDrawAggregatePath)
+				{
+					PathProgressSamples.Add(Pattern.ConvertPatternToPathProgress(PatternProgress));
+				}
+				const bool bDrawSampledPattern = Pattern.DebugDrawUsingStepSamples() && PathedMovementMode->bAllowPatternDebugDrawing;
+				if (bDrawSampledPattern)
+				{
+					const FTransform PatternRelativeTransform = Pattern.CalcMaskedTargetRelativeTransform(PatternProgress, PreviousStepPatternTransforms[PatternIdx]);
+
+					const FTransform PatternTransform = Chaos::FRigidTransform3::MultiplyNoScale(PatternRelativeTransform, LocalToWorld);
+					DebugLines.Emplace(PreviousStepPatternTransforms[PatternIdx].GetLocation(), PatternTransform.GetLocation(), Pattern.PatternDebugDrawColor, 1.f);
+
+					PreviousStepPatternTransforms[PatternIdx] = PatternTransform;
+					DebugBoundsBuilder += PatternTransform.GetLocation();
+				}
+			}
+		}
+
+		// Aggregate Path
+		if (PathedMovementMode->bDebugDrawAggregatePath)
+		{
+			// Sort overall path progress samples, they may have been added out of order when iterating over the patterns
+			PathProgressSamples.Sort();
+
+			// Establish the first point on the aggregate path
+			FVector PreviousStepAggregateLocation = Chaos::FRigidTransform3::MultiplyNoScale(PathedMovementMode->CalcTargetRelativeTransform(0.0f), LocalToWorld).GetLocation();
+			DebugBoundsBuilder += PreviousStepAggregateLocation;
+			float PreviousPathProgressSample = 0.0f;
+			// Draw the rest of the aggregate path, at the overall progress samples deduced from the sampling of path patterns
+			// and skip samples that are too close to the previous one we drew
+			for (float PathProgressSample : PathProgressSamples)
+			{
+				if (FMath::IsNearlyEqual(PathProgressSample, PreviousPathProgressSample)) // Should we look at the difference in transform location instead?
+				{
+					continue;
+				}
+				const float OverallPathProgress = PathProgressSample;
+				const FVector StepLocation = Chaos::FRigidTransform3::MultiplyNoScale(PathedMovementMode->CalcTargetRelativeTransform(OverallPathProgress), LocalToWorld).GetLocation();
+				DebugLines.Emplace(PreviousStepAggregateLocation, StepLocation, PathedMovementMode->PathDebugDrawColor, 2.f);
+				PreviousStepAggregateLocation = StepLocation;
+				DebugBoundsBuilder += StepLocation;
+				PreviousPathProgressSample = PathProgressSample;
+			}
+		}
+	}
+
+	Bounds = DebugBoundsBuilder;
+	Bounds = Bounds.ExpandBy(25.);
+
+	Super::CreateRenderState_Concurrent(Context);
+}
+
+FDebugRenderSceneProxy* UChaosPathedMovementDebugDrawComponent::CreateDebugSceneProxy()
+{
+	if (DebugLines.IsEmpty() &&
+		DebugDashedLines.IsEmpty() &&
+		DebugArrowLines.IsEmpty() &&
+		DebugSpheres.IsEmpty() &&
+		DebugStars.IsEmpty())
+	{
+		return nullptr;
+	}
+
+	FChaosPathedMovementDebugRenderSceneProxy* DebugSceneProxy = new FChaosPathedMovementDebugRenderSceneProxy(*this);
+    DebugSceneProxy->DrawType = FDebugRenderSceneProxy::SolidAndWireMeshes;
+    DebugSceneProxy->ViewFlagName = TEXT("Splines");
+    DebugSceneProxy->ViewFlagIndex = uint32(FEngineShowFlags::FindIndexByName(*DebugSceneProxy->ViewFlagName));
+	
+	DebugSceneProxy->Lines = DebugLines;
+	DebugSceneProxy->DashedLines = DebugDashedLines;
+	DebugSceneProxy->ArrowLines = DebugArrowLines;
+	DebugSceneProxy->Spheres = DebugSpheres;
+	DebugSceneProxy->Stars = DebugStars;
+
+	// We can now get rid of our debug primitives since we transfered them to the scene proxy
+	// CreateDebugSceneProxy is called by UDebugDrawComponent::CreateRenderState_Concurrent which will
+	// have caused UChaosPathedMovementDebugDrawComponent::CreateRenderState_Concurrent to recreate
+	// them next time there is a need to update those
+	DebugLines.Reset();
+	DebugDashedLines.Reset();
+	DebugArrowLines.Reset();
+	DebugSpheres.Reset();
+	DebugStars.Reset();
+	
+	return DebugSceneProxy;
+}
+
+void UChaosPathedMovementDebugDrawComponent::UpdatePreviewMeshComp(bool bForce)
+{
+	const UMoverComponent* MoverComp = GetMoverComponent();
+	if (!MoverComp)
+	{
+		return;
+	}
+
+	UMeshComponent* MeshRoot = GetOwnerMeshRoot();
+	if (MeshRoot && GetWorld() && !GetWorld()->IsGameWorld())
+	{
+		bool bDoUpdate = false;
+		
+    	if (ProgressPreviewMeshComp)
+    	{
+		    if (ProgressPreviewMeshComp->GetClass() == MeshRoot->GetClass() && bForce)
+		    {
+		    	bDoUpdate = true;
+
+		    	// Copy all the properties on the root mesh comp (except attachments)
+		    	for (TFieldIterator<FProperty> PropIter(MeshRoot->GetClass()); PropIter; ++PropIter)
+		    	{
+				    if (PropIter->GetFName() == TEXT("AttachChildren") || PropIter->GetFName() == TEXT("AttachParent"))
+				    {
+					    continue;
+				    }
+		    		PropIter->CopyCompleteValue_InContainer(ProgressPreviewMeshComp, MeshRoot);
+		    	}
+		    }
+		    else
+		    {
+			    DestroyProgressPreviewMeshComp();
+		    }
+    	}
+
+		if (!ProgressPreviewMeshComp)
+    	{
+    		ProgressPreviewMeshComp = DuplicateObject<UMeshComponent>(MeshRoot, this, TEXT("ProgressPreviewMeshComp"));
+    		// ProgressPreviewMeshComp->SetupAttachment(this);
+			ProgressPreviewMeshComp->bAutoRegister = false;
+
+			bDoUpdate = true;
+    	}
+
+		ensure(ProgressPreviewMeshComp->GetNumChildrenComponents() == 0);
+
+		if (!ProgressPreviewMeshComp->IsRegistered())
+		{
+			ProgressPreviewMeshComp->RegisterComponent();
+		}
+
+		if (bDoUpdate)
+		{
+			//@todo DanH: Double-clicking the preview comp in the level editor can still select the component - how do I disable that?
+			ProgressPreviewMeshComp->CreationMethod = EComponentCreationMethod::UserConstructionScript;
+			ProgressPreviewMeshComp->SetSimulatePhysics(false);
+			ProgressPreviewMeshComp->SetCollisionEnabled(ECollisionEnabled::NoCollision);
+			ProgressPreviewMeshComp->bIsEditorOnly = true;
+			ProgressPreviewMeshComp->SetHiddenInGame(true);
+
+			const TArray<UActorComponent*> DebugDrawInterfaceComponents = GetOwner()->GetComponentsByInterface(UChaosPathedMovementDebugDrawInterface::StaticClass());
+			if (!DebugDrawInterfaceComponents.IsEmpty())
+			{
+				DebugDrawInterfaceObject = DebugDrawInterfaceComponents[0];
+			}
+
+			const UChaosPathedMovementMode* PathedMovementMode = MoverComp->FindMode_Mutable<UChaosPathedMovementMode>(MoverComp->StartingMovementMode);			
+			if (PathedMovementMode && DebugDrawInterfaceObject && IChaosPathedMovementDebugDrawInterface::Execute_ShouldDisplayProgressPreviewMesh(DebugDrawInterfaceObject))
+			{
+				ProgressPreviewMeshComp->SetVisibility(true);
+				
+				float OverallPathProgress = IChaosPathedMovementDebugDrawInterface::Execute_GetPreviewMeshOverallPathProgress(DebugDrawInterfaceObject);
+				const FTransform PreviewMeshRelativeTransform = PathedMovementMode->CalcTargetRelativeTransform(OverallPathProgress);
+
+				// Calculate the location without scale, then tack on the source mesh's scale
+				FTransform PreviewMeshWorldTransform = Chaos::FRigidTransform3::MultiplyNoScale(PreviewMeshRelativeTransform, MeshRoot->GetComponentTransform());
+				PreviewMeshWorldTransform.SetScale3D(MeshRoot->GetComponentScale());
+				ProgressPreviewMeshComp->SetWorldTransform(PreviewMeshWorldTransform);
+
+				UMaterialInterface* ProgressPreviesMeshMaterial = IChaosPathedMovementDebugDrawInterface::Execute_GetProgressPreviewMeshMaterial(DebugDrawInterfaceObject);
+				if (ProgressPreviesMeshMaterial)
+				{
+					const int32 NumMaterials = ProgressPreviewMeshComp->GetNumMaterials();
+					for (int MaterialIdx = 0; MaterialIdx < NumMaterials; ++MaterialIdx)
+					{
+						ProgressPreviewMeshComp->SetMaterial(MaterialIdx, ProgressPreviesMeshMaterial);
+					}
+				}
+			}
+			else
+			{
+				ProgressPreviewMeshComp->SetVisibility(false);
+			}
+		}
+	}
+	else if (ProgressPreviewMeshComp)
+	{
+		DestroyProgressPreviewMeshComp();
+	}
+}
+
+void UChaosPathedMovementDebugDrawComponent::DestroyProgressPreviewMeshComp()
+{
+	if (ProgressPreviewMeshComp)
+	{
+		ProgressPreviewMeshComp->DestroyComponent();
+		ProgressPreviewMeshComp = nullptr;
+	}
+}
+
+UMeshComponent* UChaosPathedMovementDebugDrawComponent::GetOwnerMeshRoot() const
+{
+	if (AActor* Owner = GetOwner())
+	{
+		return Cast<UMeshComponent>(Owner->GetRootComponent());
+	}
+	return nullptr;
+}
+
+UMoverComponent* UChaosPathedMovementDebugDrawComponent::GetMoverComponent() const
+{
+	AActor* OwnerActor = GetOwner();
+	return OwnerActor ? OwnerActor->FindComponentByClass<UMoverComponent>() : nullptr;
+}
+
+#if WITH_EDITOR
+void UChaosPathedMovementDebugDrawComponent::HandleObjectPropertyChanged(UObject* Object, FPropertyChangedEvent& PropertyChangedEvent)
+{
+	if (UMeshComponent* MeshComponent = Cast<UMeshComponent>(Object);
+		MeshComponent && MeshComponent == GetOwner()->GetRootComponent())
+	{
+		UpdatePreviewMeshComp();		
+	}
+}
+#endif
diff --git a/Engine/Plugins/Experimental/ChaosMover/Source/ChaosMover/Private/PathedMovement/ChaosPathedMovementMode.cpp b/Engine/Plugins/Experimental/ChaosMover/Source/ChaosMover/Private/PathedMovement/ChaosPathedMovementMode.cpp
--- a/Engine/Plugins/Experimental/ChaosMover/Source/ChaosMover/Private/PathedMovement/ChaosPathedMovementMode.cpp
+++ b/Engine/Plugins/Experimental/ChaosMover/Source/ChaosMover/Private/PathedMovement/ChaosPathedMovementMode.cpp
@@ -109,7 +109,7 @@ void UChaosPathedMovementMode::SimulationTick_Implementation(const FSimulationTi
 	bool bIsInReverseJustChanged = PreviousProperties.bWantsReversePlayback != CurrentProperties.bWantsReversePlayback;
 	if (!OutputPathedMovemenState.bHasFinished && bIsInReverseJustChanged)
 	{
-		OutputPathedMovemenState.bIsPathFractionIncreasing = !OutputPathedMovemenState.bIsPathFractionIncreasing;
+		OutputPathedMovemenState.bIsPathProgressionIncreasing = !OutputPathedMovemenState.bIsPathProgressionIncreasing;
 	}
 
 	// Reset bHasFinished if we are restarting
@@ -117,12 +117,12 @@ void UChaosPathedMovementMode::SimulationTick_Implementation(const FSimulationTi
 	{
 		OutputPathedMovemenState.bHasFinished = false;
 		// We initialize the direction of motion to respect reverse playback on start
-		OutputPathedMovemenState.bIsPathFractionIncreasing = !CurrentProperties.bWantsReversePlayback;
+		OutputPathedMovemenState.bIsPathProgressionIncreasing = !CurrentProperties.bWantsReversePlayback;
 		if (CurrentProperties.bWantsOneWayPlayback)
 		{
-			if (OutputPathedMovemenState.CurrentPathFraction == (OutputPathedMovemenState.bIsPathFractionIncreasing ? 1.0f : 0.0f))
+			if (OutputPathedMovemenState.CurrentPathProgress == (OutputPathedMovemenState.bIsPathProgressionIncreasing ? 1.0f : 0.0f))
 			{
-				OutputPathedMovemenState.CurrentPathFraction = 1.0f - OutputPathedMovemenState.CurrentPathFraction;
+				OutputPathedMovemenState.CurrentPathProgress = 1.0f - OutputPathedMovemenState.CurrentPathProgress;
 			}
 		}
 	}
@@ -138,92 +138,92 @@ void UChaosPathedMovementMode::SimulationTick_Implementation(const FSimulationTi
 	{
 		const float FrameDuration = Params.TimeStep.StepMs * 0.001f;
 		const float PlaybackSpeed = (OneWayPlaybackDuration >= 0.0f) ? (1.0f / OneWayPlaybackDuration) : 0.0f;
-		const float Direction = OutputPathedMovemenState.bIsPathFractionIncreasing ? 1.0f : -1.0f;
-		const float DeltaPathFraction = Direction * PlaybackSpeed * FrameDuration;
-		const float NewPathFraction = OutputPathedMovemenState.CurrentPathFraction + DeltaPathFraction;
+		const float Direction = OutputPathedMovemenState.bIsPathProgressionIncreasing ? 1.0f : -1.0f;
+		const float DeltaPathProgress = Direction * PlaybackSpeed * FrameDuration;
+		const float NewPathProgress = OutputPathedMovemenState.CurrentPathProgress + DeltaPathProgress;
 
 		if (CurrentProperties.bWantsOneWayPlayback)
 		{
 			// One way forward playback
-			if (OutputPathedMovemenState.bIsPathFractionIncreasing)
+			if (OutputPathedMovemenState.bIsPathProgressionIncreasing)
 			{
-				if (NewPathFraction >= 1.0f)
+				if (NewPathProgress >= 1.0f)
 				{
 					if (CurrentProperties.bWantsLoopingPlayback)
 					{
-						OutputPathedMovemenState.CurrentPathFraction = FMath::Fmod(NewPathFraction, 1.0f);
+						OutputPathedMovemenState.CurrentPathProgress = FMath::Fmod(NewPathProgress, 1.0f);
 						bJustRestarted = true;
 					}
 					else
 					{
-						OutputPathedMovemenState.CurrentPathFraction = 1.0f;
+						OutputPathedMovemenState.CurrentPathProgress = 1.0f;
 						OutputPathedMovemenState.bHasFinished = true;
 						bHasJustFinished = true;
 					}
 				}
 				else
 				{
-					OutputPathedMovemenState.CurrentPathFraction = NewPathFraction;
+					OutputPathedMovemenState.CurrentPathProgress = NewPathProgress;
 				}
 			}
 			else // One way reverse playback
 			{
-				if (NewPathFraction <= 0.0f)
+				if (NewPathProgress <= 0.0f)
 				{
 					if (CurrentProperties.bWantsLoopingPlayback)
 					{
-						OutputPathedMovemenState.CurrentPathFraction = 1.0f - FMath::Fmod(-NewPathFraction, 1.0f);
+						OutputPathedMovemenState.CurrentPathProgress = 1.0f - FMath::Fmod(-NewPathProgress, 1.0f);
 						bJustRestarted = true;
 					}
 					else
 					{
-						OutputPathedMovemenState.CurrentPathFraction = 0.0f;
+						OutputPathedMovemenState.CurrentPathProgress = 0.0f;
 						OutputPathedMovemenState.bHasFinished = true;
 						bHasJustFinished = true;
 					}
 				}
 				else
 				{
-					OutputPathedMovemenState.CurrentPathFraction = NewPathFraction;
+					OutputPathedMovemenState.CurrentPathProgress = NewPathProgress;
 				}
 			}
 		}
 		else
 		{
 			// Round-trip playback
-			if (OutputPathedMovemenState.bIsPathFractionIncreasing && NewPathFraction >= 1.0f)
+			if (OutputPathedMovemenState.bIsPathProgressionIncreasing && NewPathProgress >= 1.0f)
 			{
 				if (!CurrentProperties.bWantsReversePlayback || CurrentProperties.bWantsLoopingPlayback)
 				{
-					OutputPathedMovemenState.bIsPathFractionIncreasing = false;
-					OutputPathedMovemenState.CurrentPathFraction = 1.0f - FMath::Fmod(NewPathFraction, 1.0f);
+					OutputPathedMovemenState.bIsPathProgressionIncreasing = false;
+					OutputPathedMovemenState.CurrentPathProgress = 1.0f - FMath::Fmod(NewPathProgress, 1.0f);
 					bHasJustBounced = true;
 				}
 				else
 				{
-					OutputPathedMovemenState.CurrentPathFraction = 1.0f;
+					OutputPathedMovemenState.CurrentPathProgress = 1.0f;
 					OutputPathedMovemenState.bHasFinished = true;
 					bHasJustFinished = true;
 				}
 			}
-			else if (!OutputPathedMovemenState.bIsPathFractionIncreasing && NewPathFraction <= 0.0f)
+			else if (!OutputPathedMovemenState.bIsPathProgressionIncreasing && NewPathProgress <= 0.0f)
 			{
 				if (CurrentProperties.bWantsReversePlayback || CurrentProperties.bWantsLoopingPlayback)
 				{
-					OutputPathedMovemenState.bIsPathFractionIncreasing = true;
-					OutputPathedMovemenState.CurrentPathFraction = FMath::Fmod(-NewPathFraction, 1.0f);
+					OutputPathedMovemenState.bIsPathProgressionIncreasing = true;
+					OutputPathedMovemenState.CurrentPathProgress = FMath::Fmod(-NewPathProgress, 1.0f);
 					bHasJustBounced = true;
 				}
 				else
 				{
-					OutputPathedMovemenState.CurrentPathFraction = 0.0f;
+					OutputPathedMovemenState.CurrentPathProgress = 0.0f;
 					OutputPathedMovemenState.bHasFinished = true;
 					bHasJustFinished = true;
 				}
 			}
 			else
 			{
-				OutputPathedMovemenState.CurrentPathFraction = NewPathFraction;
+				OutputPathedMovemenState.CurrentPathProgress = NewPathProgress;
 			}
 		}
 	}
@@ -258,7 +258,7 @@ void UChaosPathedMovementMode::SimulationTick_Implementation(const FSimulationTi
 		OnPlaybackBounced(Params.TimeStep.BaseSimTimeMs);
 	}
 
-	const Chaos::FRigidTransform3 TargetRelativeTransform = CalcTargetRelativeTransform(OutputPathedMovemenState.CurrentPathFraction);
+	const Chaos::FRigidTransform3 TargetRelativeTransform = CalcTargetRelativeTransform(OutputPathedMovemenState.CurrentPathProgress);
 
 	FMoverDefaultSyncState& OutputDefaultSyncState = OutputState.SyncState.SyncStateCollection.FindOrAddMutableDataByType<FMoverDefaultSyncState>();
 	
@@ -398,7 +398,7 @@ void UChaosPathedMovementMode::SetOneWayTripDuration_BeginPlayOnly(float NewDura
 	}
 }
 
-Chaos::FRigidTransform3 UChaosPathedMovementMode::CalcTargetRelativeTransform(float ProgressAmt) const
+Chaos::FRigidTransform3 UChaosPathedMovementMode::CalcTargetRelativeTransform(float OverallPathProgress) const
 {
 	FTransform TargetRelativeTransform = FTransform::Identity;
 	
@@ -406,7 +406,7 @@ Chaos::FRigidTransform3 UChaosPathedMovementMode::CalcTargetRelativeTransform(fl
 	{
 		if (PathPattern)
 		{
-			const FTransform PatternTransform = PathPattern->CalcTargetRelativeTransform(ProgressAmt, TargetRelativeTransform);
+			const FTransform PatternTransform = PathPattern->CalcTargetRelativeTransform(OverallPathProgress, TargetRelativeTransform);
         	TargetRelativeTransform.Accumulate(PatternTransform);
 		}
 	}
diff --git a/Engine/Plugins/Experimental/ChaosMover/Source/ChaosMover/Private/PathedMovement/ChaosPathedMovementPatternBase.cpp b/Engine/Plugins/Experimental/ChaosMover/Source/ChaosMover/Private/PathedMovement/ChaosPathedMovementPatternBase.cpp
--- a/Engine/Plugins/Experimental/ChaosMover/Source/ChaosMover/Private/PathedMovement/ChaosPathedMovementPatternBase.cpp
+++ b/Engine/Plugins/Experimental/ChaosMover/Source/ChaosMover/Private/PathedMovement/ChaosPathedMovementPatternBase.cpp
@@ -56,27 +56,25 @@ FRotator ApplyAxisMask(const FRotator& Unmasked, EChaosPatternAxisMaskFlags Flag
 	return MaskedResult;
 }
 
-FTransform UChaosPathedMovementPatternBase::CalcTargetRelativeTransform(float OverallPathProgress, const FTransform& CurTargetTransform) const
+float UChaosPathedMovementPatternBase::ConvertPathToPatternProgress(float OverallPathProgress) const
 {
-	//@todo DanH: When doing debug drawing, when a pattern is shorter and/or loops more, it loses resolution. Can't just always calculate as though OverallPathProgress == PatternProgress
-	//		though because we still depend on the calculated result. Maybe I shouldn't bother with that and just let the aggregate one calc on its own?
-	
-	if (NumLoopsPerPath <= 0 || StartAtPathProgress >= EndAtPathProgress || OverallPathProgress <= StartAtPathProgress)
+	// If you update this function, please also update ConvertPatternToPathProgress to support debug draws
+	if (NumLoopsPerPath <= 0 || StartAtPathProgress >= EndAtPathProgress || OverallPathProgress < StartAtPathProgress)
 	{
-		return FTransform();
+		return -1.0f;
 	}
-	
+
 	// How far into the current loop of this specific pattern are we?
 	const float PathProgressSinceStart = FMath::Min(EndAtPathProgress, OverallPathProgress) - StartAtPathProgress;
 	const float PathProgressPerPatternLoop = (EndAtPathProgress - StartAtPathProgress) / NumLoopsPerPath;
 	float CurLoopPathProgress = FMath::Fmod(PathProgressSinceStart, PathProgressPerPatternLoop);
 	if (CurLoopPathProgress == 0.f && PathProgressSinceStart > 0.f)
 	{
-		// Treat progress that matches the per-loop span exactly as 100%, not 0%
+		// Treat path progress that matches the per-loop span exactly as 100%, not 0%
 		CurLoopPathProgress = PathProgressPerPatternLoop;
 	}
 
-	// If each loop is there and back, the actual progress needs to flip after completing half the span
+	// If each loop is there and back, the actual path progress needs to flip after completing half the span
 	if (!IsOneWay)
 	{
 		// ThereAndBack means we actually progress twice as fast as a OneShot
@@ -89,8 +87,44 @@ FTransform UChaosPathedMovementPatternBase::CalcTargetRelativeTransform(float Ov
 		}
 	}
 
-	// Now convert from overall path progress to progress for this pattern
+	// Now convert from overall path progress to pattern progress
 	const float PatternProgress = CurLoopPathProgress / PathProgressPerPatternLoop;
+	const float OffsettedPatternProgress = FMath::Fmod(PatternProgress + LoopOffset, 1.0f);
+
+	return OffsettedPatternProgress;
+}
+
+float UChaosPathedMovementPatternBase::ConvertPatternToPathProgress(float PatternProgress) const
+{
+	if (NumLoopsPerPath <= 0 || StartAtPathProgress >= EndAtPathProgress)
+	{
+		return -1.0f;
+	}
+
+	const float OffsettedPatternProgress = FMath::Fmod(FMath::Clamp(PatternProgress, 0.0f, 1.0f) - LoopOffset, 1.0f);
+	const float PathProgressPerPatternLoop = (EndAtPathProgress - StartAtPathProgress) / NumLoopsPerPath;
+	float CurLoopPathProgress = PatternProgress * PathProgressPerPatternLoop;
+	if (!IsOneWay)
+	{
+		CurLoopPathProgress /= 2.f;
+	}
+
+	return StartAtPathProgress + (EndAtPathProgress - StartAtPathProgress) * CurLoopPathProgress;
+}
+
+FTransform UChaosPathedMovementPatternBase::CalcTargetRelativeTransform(float OverallPathProgress, const FTransform& CurTargetTransform) const
+{
+	float PatternProgress = ConvertPathToPatternProgress(OverallPathProgress);
+	return CalcMaskedTargetRelativeTransform(PatternProgress, CurTargetTransform);
+}
+
+FTransform UChaosPathedMovementPatternBase::CalcMaskedTargetRelativeTransform(float PatternProgress, const FTransform& CurTargetTransform) const
+{
+	if (PatternProgress < 0.0f)
+	{
+		return FTransform();
+	}
+
 	const FTransform UnmaskedTarget = CalcUnmaskedTargetRelativeTransform(PatternProgress, CurTargetTransform);
 
 	if (bOrientComponentToPath)
diff --git a/Engine/Plugins/Experimental/ChaosMover/Source/ChaosMover/Private/PathedMovement/ChaosPathedMovementTypes.cpp b/Engine/Plugins/Experimental/ChaosMover/Source/ChaosMover/Private/PathedMovement/ChaosPathedMovementTypes.cpp
--- a/Engine/Plugins/Experimental/ChaosMover/Source/ChaosMover/Private/PathedMovement/ChaosPathedMovementTypes.cpp
+++ b/Engine/Plugins/Experimental/ChaosMover/Source/ChaosMover/Private/PathedMovement/ChaosPathedMovementTypes.cpp
@@ -78,12 +78,12 @@ bool FChaosPathedMovementState::NetSerialize(FArchive& Ar, UPackageMap* Map, boo
 
 	static FChaosPathedMovementState Defaults;
 
-	bool HasCurrentPathFraction;
+	bool HasCurrentPathProgress;
 	bool HasLatestMovementStartFrame;
 
 	if (Ar.IsSaving())
 	{
-		HasCurrentPathFraction = CurrentPathFraction != Defaults.CurrentPathFraction;
+		HasCurrentPathProgress = CurrentPathProgress != Defaults.CurrentPathProgress;
 		HasLatestMovementStartFrame = LatestMovementStartFrame != Defaults.LatestMovementStartFrame;
 	}
 	else
@@ -92,15 +92,15 @@ bool FChaosPathedMovementState::NetSerialize(FArchive& Ar, UPackageMap* Map, boo
 		*this = Defaults;
 	}
 
-	Ar.SerializeBits(&HasCurrentPathFraction, 1);
+	Ar.SerializeBits(&HasCurrentPathProgress, 1);
 	Ar.SerializeBits(&HasLatestMovementStartFrame, 1);
-	Ar.SerializeBits(&bIsPathFractionIncreasing, 1);
+	Ar.SerializeBits(&bIsPathProgressionIncreasing, 1);
 	Ar.SerializeBits(&bHasFinished, 1);
 	PropertiesInEffect.NetSerialize(Ar);
 
-	if (HasCurrentPathFraction)
+	if (HasCurrentPathProgress)
 	{
-		Ar << CurrentPathFraction;
+		Ar << CurrentPathProgress;
 	}
 	if (HasLatestMovementStartFrame)
 	{
@@ -114,9 +114,9 @@ void FChaosPathedMovementState::ToString(FAnsiStringBuilderBase& Out) const
 {
 	FMoverDataStructBase::ToString(Out);
 
-	Out.Appendf("Current Path Fraction: %.2f | ", CurrentPathFraction);
+	Out.Appendf("Current Path Progress: %.2f | ", CurrentPathProgress);
 	Out.Appendf("Latest Movement Start Frame: %d | ", LatestMovementStartFrame);
-	Out.Appendf("Is Path Fraction Increasing: %s | ", bIsPathFractionIncreasing ? TEXT("True") : TEXT("False"));
+	Out.Appendf("Is Path Progress Increasing: %s | ", bIsPathProgressionIncreasing ? TEXT("True") : TEXT("False"));
 	Out.Appendf("Has Finished: %s | ", bHasFinished ? TEXT("True") : TEXT("False"));
 	PropertiesInEffect.ToString(Out);
 }
@@ -125,13 +125,13 @@ bool FChaosPathedMovementState::ShouldReconcile(const FMoverDataStructBase& Auth
 {
 	const FChaosPathedMovementState& TypedAuthority = static_cast<const FChaosPathedMovementState&>(AuthorityState);
 
-	static constexpr float PathFractionTolerance = 0.05f;
+	static constexpr float PathProgressTolerance = 0.05f;
 
 	if (PropertiesInEffect != TypedAuthority.PropertiesInEffect
 		|| LatestMovementStartFrame != TypedAuthority.LatestMovementStartFrame
-		|| bIsPathFractionIncreasing != TypedAuthority.bIsPathFractionIncreasing
+		|| bIsPathProgressionIncreasing != TypedAuthority.bIsPathProgressionIncreasing
 		|| bHasFinished != TypedAuthority.bHasFinished
-		|| !FMath::IsNearlyEqual(CurrentPathFraction, TypedAuthority.CurrentPathFraction, PathFractionTolerance))
+		|| !FMath::IsNearlyEqual(CurrentPathProgress, TypedAuthority.CurrentPathProgress, PathProgressTolerance))
 	{
 		return true;
 	}
@@ -146,8 +146,8 @@ void FChaosPathedMovementState::Interpolate(const FMoverDataStructBase& From, co
 
 	PropertiesInEffect = (Pct < .5f) ? TypedFrom.PropertiesInEffect : TypedTo.PropertiesInEffect;
 	LatestMovementStartFrame = (Pct < .5f) ? TypedFrom.LatestMovementStartFrame : TypedTo.LatestMovementStartFrame;
-	CurrentPathFraction = FMath::Lerp(TypedFrom.CurrentPathFraction, TypedTo.CurrentPathFraction, Pct);
-	bIsPathFractionIncreasing = (Pct < .5f) ? TypedFrom.bIsPathFractionIncreasing : TypedTo.bIsPathFractionIncreasing;
+	CurrentPathProgress = FMath::Lerp(TypedFrom.CurrentPathProgress, TypedTo.CurrentPathProgress, Pct);
+	bIsPathProgressionIncreasing = (Pct < .5f) ? TypedFrom.bIsPathProgressionIncreasing : TypedTo.bIsPathProgressionIncreasing;
 	bHasFinished = (Pct < .5f) ? TypedFrom.bHasFinished : TypedTo.bHasFinished;
 }
 
diff --git a/Engine/Plugins/Experimental/ChaosMover/Source/ChaosMover/Private/PathedMovement/ChaosPointMovementPathPattern.cpp b/Engine/Plugins/Experimental/ChaosMover/Source/ChaosMover/Private/PathedMovement/ChaosPointMovementPathPattern.cpp
--- a/Engine/Plugins/Experimental/ChaosMover/Source/ChaosMover/Private/PathedMovement/ChaosPointMovementPathPattern.cpp
+++ b/Engine/Plugins/Experimental/ChaosMover/Source/ChaosMover/Private/PathedMovement/ChaosPointMovementPathPattern.cpp
@@ -6,7 +6,8 @@
 #include "Kismet/KismetMathLibrary.h"
 #include "ChaosMover/ChaosMoverSimulation.h"
 #include "Framework/Threading.h"
-//#include "ChaosMover/PathedMovement/ChaosPathedPhysicsDebugDrawComponent.h"
+#include "ChaosMover/PathedMovement/ChaosPathedMovementMode.h"
+#include "ChaosMover/PathedMovement/ChaosPathedMovementDebugDrawComponent.h"
 
 void UChaosPointMovementPathPattern::InitializePattern(UChaosMoverSimulation* InSimulation)
 {
@@ -24,13 +25,11 @@ void UChaosPointMovementPathPattern::PostEditChangeProperty(struct FPropertyChan
 }
 #endif
 
-
-/*
-void UChaosPointMovementPathPattern::AppendDebugDrawElements(UChaosPathedPhysicsDebugDrawComponent& DebugDrawComp, FBoxSphereBounds::Builder& InOutDebugBoundsBuilder)
+void UChaosPointMovementPathPattern::AppendDebugDrawElements(UChaosPathedMovementDebugDrawComponent& DebugDrawComp, FBoxSphereBounds::Builder& InOutDebugBoundsBuilder)
 {
 	if (!PathPoints.IsEmpty())
 	{
-		const FTransform& PathOriginTransform = GetPathedMoverComp().GetPathOriginTransform();
+		const FTransform& PathOriginTransform = Simulation ? Simulation->GetPathedMovementBasis() : FTransform::Identity;
 		FVector PrevPointLocation = PathOriginTransform.GetLocation();
 		for (const FChaosPointMovementPathPoint& PathPoint : PathPoints)
 		{
@@ -55,7 +54,6 @@ void UChaosPointMovementPathPattern::AppendDebugDrawElements(UChaosPathedPhysics
 		DebugDrawComp.DebugSpheres.Emplace(4.f, PrevPointLocation, PatternDebugDrawColor);
 	}
 }
-*/
 
 FTransform UChaosPointMovementPathPattern::CalcUnmaskedTargetRelativeTransform(float PatternProgress, const FTransform& CurTargetTransform) const
 {
@@ -66,7 +64,6 @@ FTransform UChaosPointMovementPathPattern::CalcUnmaskedTargetRelativeTransform(f
 	if (PathPoints.IsValidIndex(DestPointIdx))
 	{
 		const FTransform& PathBasis = Simulation ? Simulation->GetPathedMovementBasis() : FTransform::Identity;
-		
 		const FChaosPointMovementPathPoint& NextPoint = PathPoints[DestPointIdx];
 		FChaosPointMovementPathPoint PrevPoint;
 		if (DestPointIdx > 0)
diff --git a/Engine/Plugins/Experimental/ChaosMover/Source/ChaosMover/Private/PathedMovement/ChaosSplineMovementPathPattern.cpp b/Engine/Plugins/Experimental/ChaosMover/Source/ChaosMover/Private/PathedMovement/ChaosSplineMovementPathPattern.cpp
--- a/Engine/Plugins/Experimental/ChaosMover/Source/ChaosMover/Private/PathedMovement/ChaosSplineMovementPathPattern.cpp
+++ b/Engine/Plugins/Experimental/ChaosMover/Source/ChaosMover/Private/PathedMovement/ChaosSplineMovementPathPattern.cpp
@@ -13,7 +13,7 @@ void UChaosSplineMovementPathPattern::InitializePattern(UChaosMoverSimulation* I
 
 	Chaos::EnsureIsInGameThreadContext();
 
-	if (!SplineComp && InSimulation)
+	if (!SplineComp)
 	{
 		UChaosPathedMovementMode& MovementMode = GetMovementMode();
 		if (UMoverComponent* MoverComponent = MovementMode.GetMoverComponent())
diff --git a/Engine/Plugins/Experimental/ChaosMover/Source/ChaosMover/Public/ChaosMover/PathedMovement/ChaosPathedMovementControllerComponent.h b/Engine/Plugins/Experimental/ChaosMover/Source/ChaosMover/Public/ChaosMover/PathedMovement/ChaosPathedMovementControllerComponent.h
--- a/Engine/Plugins/Experimental/ChaosMover/Source/ChaosMover/Public/ChaosMover/PathedMovement/ChaosPathedMovementControllerComponent.h
+++ b/Engine/Plugins/Experimental/ChaosMover/Source/ChaosMover/Public/ChaosMover/PathedMovement/ChaosPathedMovementControllerComponent.h
@@ -6,6 +6,7 @@
 #include "Delegates/Delegate.h"
 #include "ChaosMover/ChaosMoverSimulationTypes.h"
 #include "ChaosMover/PathedMovement/ChaosPathedMovementTypes.h"
+#include "ChaosMover/PathedMovement/ChaosPathedMovementDebugDrawComponent.h"
 
 #include "ChaosPathedMovementControllerComponent.generated.h"
 
@@ -14,6 +15,7 @@ namespace Chaos
 class FPBDRigidsSolver;
 }
 class UMoverComponent;
+class UChaosPathedMovementDebugDrawComponent;
 
 // Fired on the game thread when it has received notification that the pathed movement has started
 DECLARE_DYNAMIC_MULTICAST_DELEGATE(FChaosMover_OnPathedMovementStarted);
@@ -53,6 +55,7 @@ enum class EChaosPathedMovementExecutionType : uint8
 UCLASS(BlueprintType, Blueprintable, meta = (BlueprintSpawnableComponent))
 class UChaosPathedMovementControllerComponent : public UActorComponent
 	, public IMoverInputProducerInterface
+	, public IChaosPathedMovementDebugDrawInterface
 {
 	GENERATED_BODY()
 
@@ -61,6 +64,8 @@ public:
 
 	// UActorComponent
 	virtual void BeginPlay() override;
+	virtual void OnRegister() override;
+	virtual void OnUnregister() override;
 
 	// IMoverInputProducerInterface
 	virtual void ProduceInput_Implementation(int32 SimTimeMs, FMoverInputCmdContext& InputCmdResult) override;
@@ -143,6 +148,27 @@ protected:
 	UFUNCTION()
 	virtual void OnPostFinalize(const FMoverSyncState& SyncState, const FMoverAuxStateContext& AuxState);
 
+protected:
+	// IChaosPathedMovementDebugDrawInterface implementation
+	virtual bool ShouldDisplayProgressPreviewMesh_Implementation() const override;
+	virtual float GetPreviewMeshOverallPathProgress_Implementation() const override;
+	virtual UMaterialInterface* GetProgressPreviewMeshMaterial_Implementation() const override;
+
+	/** When true and the root component of the actor is a mesh, a duplicate mesh will be shown at PreviewMeshProgress along the starting/default path */
+	UPROPERTY(EditAnywhere, Category = "Pathed Movement|Debug")
+	bool bDisplayProgressPreviewMesh = true;
+
+	/** How far along the starting/default path to preview the controlled mesh */
+	UPROPERTY(EditAnywhere, Category = "Pathed Movement|Debug", meta = (EditCondition = bDisplayProgressPreviewMesh, EditConditionHides, UIMin = 0, ClampMin = 0, UIMax = 1, ClampMax = 1))
+	float PreviewMeshOverallPathProgress = 1.f;
+
+	/** The material to apply to the preview mesh displayed along the path at PreviewMeshProgress (leave empty for an exact duplicate) */
+	UPROPERTY(EditAnywhere, Category = "Pathed Movement|Debug", meta = (EditCondition = bDisplayProgressPreviewMesh, EditConditionHides))
+	TObjectPtr<UMaterialInterface> PreviewMeshMaterial;
+
+	UPROPERTY(Transient)
+	TObjectPtr<UChaosPathedMovementDebugDrawComponent> DebugDrawComp;
+
 private:
 	// This returns the earliest future frame index at which a requested change can occur,
 	// giving enough time for the change to be received on all end points before occurring
diff --git a/Engine/Plugins/Experimental/ChaosMover/Source/ChaosMover/Public/ChaosMover/PathedMovement/ChaosPathedMovementDebugDrawComponent.h b/Engine/Plugins/Experimental/ChaosMover/Source/ChaosMover/Public/ChaosMover/PathedMovement/ChaosPathedMovementDebugDrawComponent.h
new file mode 100644
--- /dev/null
+++ b/Engine/Plugins/Experimental/ChaosMover/Source/ChaosMover/Public/ChaosMover/PathedMovement/ChaosPathedMovementDebugDrawComponent.h
@@ -0,0 +1,76 @@
+// Copyright Epic Games, Inc. All Rights Reserved.
+
+#pragma once
+
+#include "Debug/DebugDrawComponent.h"
+
+#include "ChaosPathedMovementDebugDrawComponent.generated.h"
+
+class UMeshComponent;
+class UMoverComponent;
+
+UINTERFACE(MinimalAPI, BlueprintType)
+class UChaosPathedMovementDebugDrawInterface : public UInterface
+{
+	GENERATED_BODY()
+};
+
+/**
+ * ChaosPathedMovementDebugDrawInterface: Controls the debug draws of pathed movement modes, showing a preview mesh at a given overall progression on the path
+ */
+class IChaosPathedMovementDebugDrawInterface : public IInterface
+{
+	GENERATED_BODY()
+
+public:
+	/** Whether pathed movement debug draws should display a preview mesh at a given overall progress of the aggregate path */
+	UFUNCTION(BlueprintNativeEvent)
+	bool ShouldDisplayProgressPreviewMesh() const;
+
+	/** How far along the path progression to preview the controlled mesh */
+	UFUNCTION(BlueprintNativeEvent)
+	float GetPreviewMeshOverallPathProgress() const;
+
+	/** The material to apply to the preview mesh displayed along the path at PreviewMeshProgress (leave empty for an exact duplicate) */
+	UFUNCTION(BlueprintNativeEvent)
+	UMaterialInterface* GetProgressPreviewMeshMaterial() const;
+};
+
+//@todo DanH: Ideally clicking anything drawn by this would select the mover comp in the editor. I think I need a whole FComponentVisualizer set up for that though?
+UCLASS()
+class UChaosPathedMovementDebugDrawComponent : public UDebugDrawComponent
+{
+	GENERATED_BODY()
+
+public:
+	UChaosPathedMovementDebugDrawComponent();
+
+	virtual void OnRegister() override;
+	virtual void OnUnregister() override;
+
+	virtual FBoxSphereBounds CalcBounds(const FTransform& LocalToWorld) const override;
+
+	TArray<FDebugRenderSceneProxy::FDebugLine> DebugLines;
+	TArray<FDebugRenderSceneProxy::FDashedLine> DebugDashedLines;
+	TArray<FDebugRenderSceneProxy::FArrowLine> DebugArrowLines;
+	TArray<FDebugRenderSceneProxy::FWireStar> DebugStars;
+	TArray<FDebugRenderSceneProxy::FSphere> DebugSpheres;
+	
+protected:
+	virtual void CreateRenderState_Concurrent(FRegisterComponentContext* Context) override;
+	virtual FDebugRenderSceneProxy* CreateDebugSceneProxy() override;
+
+	void UpdatePreviewMeshComp(bool bForce = false);
+	void DestroyProgressPreviewMeshComp();
+	UMeshComponent* GetOwnerMeshRoot() const;
+	UMoverComponent* GetMoverComponent() const;
+
+#if WITH_EDITOR
+	void HandleObjectPropertyChanged(UObject* Object, FPropertyChangedEvent& PropertyChangedEvent);
+#endif
+
+	UPROPERTY(Transient)
+	TObjectPtr<UMeshComponent> ProgressPreviewMeshComp = nullptr;
+
+	TObjectPtr<UObject> DebugDrawInterfaceObject;
+};
\ No newline at end of file
diff --git a/Engine/Plugins/Experimental/ChaosMover/Source/ChaosMover/Public/ChaosMover/PathedMovement/ChaosPathedMovementMode.h b/Engine/Plugins/Experimental/ChaosMover/Source/ChaosMover/Public/ChaosMover/PathedMovement/ChaosPathedMovementMode.h
--- a/Engine/Plugins/Experimental/ChaosMover/Source/ChaosMover/Public/ChaosMover/PathedMovement/ChaosPathedMovementMode.h
+++ b/Engine/Plugins/Experimental/ChaosMover/Source/ChaosMover/Public/ChaosMover/PathedMovement/ChaosPathedMovementMode.h
@@ -40,7 +40,7 @@ public:
 	virtual const Chaos::FPBDJointSettings& GetConstraintSettings() const override;
 
 	// IChaosProgressingPathedMovementModeInterface
-	virtual Chaos::FRigidTransform3 CalcTargetRelativeTransform(float ProgressAmt) const override;
+	virtual Chaos::FRigidTransform3 CalcTargetRelativeTransform(float OverallPathProgress) const override;
 
 
 #if WITH_EDITOR
@@ -124,6 +124,20 @@ private:
 
 	Chaos::FPBDJointSettings JointSettings;
 
+	/** True to draw debug lines for the aggregate path resulting from all the patterns in the path (i.e. the path that will actually be followed in game) */
+	UPROPERTY(EditAnywhere, Category = "Pathed Movement Mode|Debug")
+	bool bDebugDrawAggregatePath = true;
+
+	/** True to allow individual patterns to draw debug representations. Can be toggled individually on each pattern as well, this is a master switch for convenience. */
+	UPROPERTY(EditAnywhere, Category = "Pathed Movement Mode|Debug")
+	bool bAllowPatternDebugDrawing = true;
+
+	/** The color used when drawing the aggregate path */
+	UPROPERTY(EditAnywhere, Category = "Pathed Movement Mode|Debug", meta = (EditCondition = bDebugDrawAggregatePath))
+	FColor PathDebugDrawColor = FColor::Blue;
+	
+	friend class UChaosPathedMovementDebugDrawComponent;
+
 #if !UE_BUILD_SHIPPING
 	FChaosPathedMovementModeDebugData DebugData;
 #endif // !UE_BUILD_SHIPPING
diff --git a/Engine/Plugins/Experimental/ChaosMover/Source/ChaosMover/Public/ChaosMover/PathedMovement/ChaosPathedMovementPatternBase.h b/Engine/Plugins/Experimental/ChaosMover/Source/ChaosMover/Public/ChaosMover/PathedMovement/ChaosPathedMovementPatternBase.h
--- a/Engine/Plugins/Experimental/ChaosMover/Source/ChaosMover/Public/ChaosMover/PathedMovement/ChaosPathedMovementPatternBase.h
+++ b/Engine/Plugins/Experimental/ChaosMover/Source/ChaosMover/Public/ChaosMover/PathedMovement/ChaosPathedMovementPatternBase.h
@@ -12,7 +12,7 @@ class UCurveFloat;
 enum class EAlphaBlendOption : uint8;
 class UChaosPathedMovementMode;
 class UChaosMoverSimulation;
-//class UChaosPathedPhysicsDebugDrawComponent;
+class UChaosPathedMovementDebugDrawComponent;
 
 //@todo DanH: What's the Flags specifier for? Does that do everything now or do I still need the Bitflags meta?
 UENUM(Flags, meta = (Bitflags))
@@ -47,7 +47,7 @@ public:
 	FTransform CalcTargetRelativeTransform(float OverallPathProgress, const FTransform& CurTargetTransform) const;
 
 	virtual bool DebugDrawUsingStepSamples() const { return bDebugDrawPattern; }
-	//virtual void AppendDebugDrawElements(UChaosPathedPhysicsDebugDrawComponent& DebugDrawComp, FBoxSphereBounds::Builder& InOutDebugBoundsBuilder) {}
+	virtual void AppendDebugDrawElements(UChaosPathedMovementDebugDrawComponent& DebugDrawComp, FBoxSphereBounds::Builder& InOutDebugBoundsBuilder) {}
 
 #if WITH_EDITOR
 	//@todo DanH: Pipe this through (I wonder if there should be a UPackage::IsDataValid() override that auto-pipes it to all UObjects in the package?)
@@ -58,11 +58,27 @@ public:
 	
 protected:
 	/**
-	* Where patterns calculate their target transforms at the given progress, optionally relative to the current aggregate target.
+	* Where patterns calculate their target transforms at the given path progress, optionally relative to the current aggregate target.
 	* Axis masking is applied to this result before it's actually used.
 	*/
 	virtual FTransform CalcUnmaskedTargetRelativeTransform(float PatternProgress, const FTransform& CurTargetTransform) const { return FTransform::Identity; }
 
+	/** Utility function that calls CalcUnmaskedTargetRelativeTransform and then applies masking. This is mostly for debug drawing, which wants to
+	*   deal with PatternProgress and not OverallPathProgress but also get the masked transform
+	*/
+	FTransform CalcMaskedTargetRelativeTransform(float PatternProgress, const FTransform& CurTargetTransform) const;
+
+	/** Converts an overall path progress to a progress on the pattern
+	*   Returns -1 if it can't do it, say, e.g. if somehow StartAtPathProgress > EndAtPathProgress
+	*/
+	float ConvertPathToPatternProgress(float OverallPathProgress) const;
+
+	/** Converts a pattern progress to one overall path possible progress
+	*	If the pattern is looping, it will return the path progress corresponding to the first loop
+	*   Note that ConvertPathToPatternProgress is not invertible, it is a surjective mapping (i.e. many-to-one mapping)
+	*/
+	float ConvertPatternToPathProgress(float PatternProgress) const;
+
 	//
 	// NOTE: Making any of these properties mutable at game time will require either replicating them (easier, but not as good) or pushing them as inputs to the PT (harder, more correct)
 	//
@@ -86,24 +102,24 @@ protected:
 	UPROPERTY(EditAnywhere, Category = "Path Pattern")
 	bool bStartAfterPreviousPattern = false;
 
-	// /**
-	//  * The desired duration (in overall path progress) for this pattern once it begins.
-	//  * If this desired duration is longer than the progress remaining when the pattern starts, it will be clamped to fill what remains
-	//  */
-	// UPROPERTY(EditAnywhere, Category = "Path Pattern", meta = (EditCondition = "bStartAfterPreviousPattern", UIMin = 0.f, ClampMin = 0.f, UIMax = 1.f, ClampMax = 1.f))
-	// float DesiredPatternDuration = 1.f;
-
-	/** The overall path progress when this pattern should begin, where 0 is the start of the path and 1 is the end. Must be less than EndAtProgress. */
+	/** The overall path progress when this pattern should begin, where 0 is the start of the path and 1 is the end. Must be less than EndAtPathProgress. */
 	UPROPERTY(EditAnywhere, Category = "Path Pattern", meta = (EditCondition = "!bStartAfterPreviousPattern", UIMin = 0.f, ClampMin = 0.f, UIMax = 1.f, ClampMax = 1.f))
 	float StartAtPathProgress = 0.f;
 
-	/** The overall path progress when this pattern should complete, where 0 is the start of the path and 1 is the end. Must be greater than StartAtProgress. */
+	/** The overall path progress when this pattern should complete, where 0 is the start of the path and 1 is the end. Must be greater than StartAtPathProgress. */
 	UPROPERTY(EditAnywhere, Category = "Path Pattern", meta = (UIMin = 0.f, ClampMin = 0.f, UIMax = 1.f, ClampMax = 1.f))
 	float EndAtPathProgress = 1.f;
 
+	/** This is an offset within a loop of the pattern, allowing to offset the start of the loop within the path. For example if this is set to 0.2,
+	*   a loop will start at 20% of the path (instead of at the start), progress to the end, then wrap around to the start and continue progressing to 20%.
+	*   If the path itself is not a loop (if start and end are not the same), then the position will jump abruptly when it wraps around.
+	*/
+	UPROPERTY(EditAnywhere, Category = "Path Pattern", meta = (UIMin = 0.f, ClampMin = 0.f, UIMax = 1.f, ClampMax = 1.f))
+	float LoopOffset = 0.f;
+
 	//@todo DanH: Maybe this should be a float and allow partial loops?
 	/**
-	 * The number of loops to complete within the active span of this pattern (i.e. between StartAtProgress and EndAtProgress)
+	 * The number of loops to complete within the active span of this pattern (i.e. between StartAtPathProgress and EndAtPathProgress)
 	 * on a single run along the full aggregate path. Setting to 0 effectively disables this pattern.
 	 */
 	UPROPERTY(EditAnywhere, Category = "Path Pattern", meta = (UIMin = 0, ClampMin = 0))
@@ -139,5 +155,5 @@ protected:
 
 	UChaosMoverSimulation* Simulation = nullptr;
 
-	//friend class UChaosPathedPhysicsDebugDrawComponent;
+	friend class UChaosPathedMovementDebugDrawComponent;
 };
diff --git a/Engine/Plugins/Experimental/ChaosMover/Source/ChaosMover/Public/ChaosMover/PathedMovement/ChaosPathedMovementTypes.h b/Engine/Plugins/Experimental/ChaosMover/Source/ChaosMover/Public/ChaosMover/PathedMovement/ChaosPathedMovementTypes.h
--- a/Engine/Plugins/Experimental/ChaosMover/Source/ChaosMover/Public/ChaosMover/PathedMovement/ChaosPathedMovementTypes.h
+++ b/Engine/Plugins/Experimental/ChaosMover/Source/ChaosMover/Public/ChaosMover/PathedMovement/ChaosPathedMovementTypes.h
@@ -93,10 +93,11 @@ struct FChaosPathedMovementState : public FMoverDataStructBase
 	GENERATED_BODY()
 
 	/** This is a value between 0 and 1 that can be directly used along the path to calculate
-	  * the current position on it, 0.0 being a the start of the path, 1.0 being at the end.
+	  * the current position on it, 0.0 being a the start of the progression along the path, 1.0 being at the end of that progression.
+	  * Note that this is not necessarily the start and end of that path as the start could be offseted relative to the path primitive (say, a spline)
 	  */
 	UPROPERTY(VisibleAnywhere, Category = ChaosMover)
-	float CurrentPathFraction = 0.0f;
+	float CurrentPathProgress = 0.0f;
 
 	/** The frame at which we last started playing, or INDEX_NONE if we were never playing */
 	UPROPERTY(VisibleAnywhere, Category = ChaosMover)
@@ -104,7 +105,7 @@ struct FChaosPathedMovementState : public FMoverDataStructBase
 
 	/** Whether traveling from 0 to 1 (true) or 1 to 0 (false) */
 	UPROPERTY(VisibleAnywhere, Category = ChaosMover)
-	bool bIsPathFractionIncreasing = true;
+	bool bIsPathProgressionIncreasing = true;
 	
 	/** Whether the previous playback finished by reaching its end point.
 	*/
@@ -299,7 +300,7 @@ struct FPathedMovementOneWayChangedEventData : public FMoverSimulationEventData
 };
 
 // Interface for mover modes which can move the controlled component
-// with a general purpose joint constraint or kinematically
+// with a physical constraint or kinematically
 UINTERFACE(MinimalAPI)
 class UChaosMovementActuationInterface : public UInterface
 {
@@ -310,17 +311,17 @@ class IChaosMovementActuationInterface
 	GENERATED_BODY()
 public:
 
-	// Whether the simulation should use a constraint to move the actor along the path, or moving the actor kinematically
+	// Whether the simulation should use a constraint to move the controlled particle, or moving it kinematically
 	virtual bool ShouldUseConstraint() const = 0;
 
 	// Constraint settings of the joint constraint, only valid when ShouldUseConstraint() is true
 	virtual const Chaos::FPBDJointSettings& GetConstraintSettings() const = 0;
 };
 
-// Interface for mover modes moving on a path, with a constraint or kinematically
-// Paths are spatial trajectories with a beginning and an end. A unique position on the path
-// can be calculated by providing a path fraction in [0,1], 0 being at the beginning of the path
-// 1 being at the end.
+// Interface for mover modes which target transform is parameterized by an overall path progression, a value in [0, 1]
+// 0 being at the beginning of the progression, 1 being at the end of that progression. A given progression should
+// always result in the same resulting transform, but the spatial trajectory resulting from the progression
+// varying from 0 to 1 may loop or revisit the same curves multiple times
 UINTERFACE(MinimalAPI)
 class UChaosPathedMovementModeInterface : public UInterface
 {
@@ -331,6 +332,6 @@ class IChaosPathedMovementModeInterface
 	GENERATED_BODY()
 public:
 
-	// Returns the transform on the path at the path fraction
-	virtual Chaos::FRigidTransform3 CalcTargetRelativeTransform(float PathFraction) const = 0;
+	// Returns the transform on the path at Progression
+	virtual Chaos::FRigidTransform3 CalcTargetRelativeTransform(float Progression) const = 0;
 };
diff --git a/Engine/Plugins/Experimental/ChaosMover/Source/ChaosMover/Public/ChaosMover/PathedMovement/ChaosPointMovementPathPattern.h b/Engine/Plugins/Experimental/ChaosMover/Source/ChaosMover/Public/ChaosMover/PathedMovement/ChaosPointMovementPathPattern.h
--- a/Engine/Plugins/Experimental/ChaosMover/Source/ChaosMover/Public/ChaosMover/PathedMovement/ChaosPointMovementPathPattern.h
+++ b/Engine/Plugins/Experimental/ChaosMover/Source/ChaosMover/Public/ChaosMover/PathedMovement/ChaosPointMovementPathPattern.h
@@ -7,6 +7,8 @@
 
 #include "ChaosPointMovementPathPattern.generated.h"
 
+class UChaosPathedMovementDebugDrawComponent;
+
 UENUM()
 enum class EChaosPointMovementLocationBasis : uint8
 {
@@ -64,7 +66,7 @@ public:
 #endif
 	
 	virtual bool DebugDrawUsingStepSamples() const override { return false; }
-	//virtual void AppendDebugDrawElements(UChaosPathedPhysicsDebugDrawComponent& DebugDrawComp, FBoxSphereBounds::Builder& InOutDebugBoundsBuilder) override;
+	virtual void AppendDebugDrawElements(UChaosPathedMovementDebugDrawComponent& DebugDrawComp, FBoxSphereBounds::Builder& InOutDebugBoundsBuilder) override;
 
 protected:
 	virtual FTransform CalcUnmaskedTargetRelativeTransform(float PatternProgress, const FTransform& CurTargetTransform) const override;
