ChaosMover: Porting EllipticalPathPattern

diff --git a/Engine/Plugins/Experimental/ChaosMover/Source/ChaosMover/Private/PathedMovement/ChaosEllipticalMovementPathPattern.cpp b/Engine/Plugins/Experimental/ChaosMover/Source/ChaosMover/Private/PathedMovement/ChaosEllipticalMovementPathPattern.cpp
new file mode 100644
--- /dev/null
+++ b/Engine/Plugins/Experimental/ChaosMover/Source/ChaosMover/Private/PathedMovement/ChaosEllipticalMovementPathPattern.cpp
@@ -0,0 +1,27 @@
+// Copyright Epic Games, Inc. All Rights Reserved.
+
+#include "ChaosMover/PathedMovement/ChaosEllipticalMovementPathPattern.h"
+
+#include UE_INLINE_GENERATED_CPP_BY_NAME(ChaosEllipticalMovementPathPattern)
+
+void UChaosEllipticalMovementPathPattern::InitializePattern(UChaosMoverSimulation* InSimulation)
+{
+	Super::InitializePattern(InSimulation);
+}
+
+FTransform UChaosEllipticalMovementPathPattern::CalcUnmaskedTargetRelativeTransform(float PatternProgress, const FTransform& CurTargetTransform) const
+{
+	// Calc the ellipse point
+	const float Angle = PatternProgress * FMath::DegreesToRadians(UsableArcAngle);
+	FVector2f TargetLocationXY(RadiusX * FMath::Cos(Angle), RadiusY * FMath::Sin(Angle));
+
+	// Scooch the center over so 0 progress means no movement
+	TargetLocationXY.X -= RadiusX;
+
+	// Spin about the path origin as desired
+	TargetLocationXY = TargetLocationXY.GetRotated(OriginAngle);
+
+	// Rotate the ellipse from the XY plane to whatever 3D plane is desired
+	const FVector TargetLocation = EllipsePlaneRotation.RotateVector(FVector(TargetLocationXY.X, TargetLocationXY.Y, 0.f));
+	return FTransform(TargetLocation);
+}
diff --git a/Engine/Plugins/Experimental/ChaosMover/Source/ChaosMover/Public/ChaosMover/PathedMovement/ChaosEllipticalMovementPathPattern.h b/Engine/Plugins/Experimental/ChaosMover/Source/ChaosMover/Public/ChaosMover/PathedMovement/ChaosEllipticalMovementPathPattern.h
new file mode 100644
--- /dev/null
+++ b/Engine/Plugins/Experimental/ChaosMover/Source/ChaosMover/Public/ChaosMover/PathedMovement/ChaosEllipticalMovementPathPattern.h
@@ -0,0 +1,39 @@
+// Copyright Epic Games, Inc. All Rights Reserved.
+
+#pragma once
+
+#include "CoreMinimal.h"
+#include "ChaosMover/PathedMovement/ChaosPathedMovementPatternBase.h"
+
+#include "ChaosEllipticalMovementPathPattern.generated.h"
+
+UCLASS()
+class UChaosEllipticalMovementPathPattern : public UChaosPathedMovementPatternBase
+{
+	GENERATED_BODY()
+	
+public:
+	virtual void InitializePattern(UChaosMoverSimulation* InSimulation) override;
+
+protected:
+	virtual FTransform CalcUnmaskedTargetRelativeTransform(float OverallPathFraction, const FTransform& CurTargetTransform) const override;
+
+	/** Radius of the ellipse along the x axis */
+	UPROPERTY(EditAnywhere, Category = EllipticalPathPattern)
+	float RadiusX = 100.f;
+	
+	UPROPERTY(EditAnywhere, Category = EllipticalPathPattern)
+	float RadiusY = 100.f;
+
+	/** The angle between the path origin (i.e. the initial location of the component/particle) and the center of the ellipse */
+	UPROPERTY(EditAnywhere, Category = EllipticalPathPattern, meta = (Units = Degrees, UIMin = 0, ClampMin = 0, UIMax = 360, ClampMax = 360))
+	float OriginAngle = 0.f;
+
+	/** The amount of the ellipsoid arc to actually use for the path (where 360 is the entire ellipse) */
+	UPROPERTY(EditAnywhere, Category = EllipticalPathPattern, meta = (Units = Degrees, UIMin = 0, ClampMin = 0, UIMax = 360, ClampMax = 360))
+	float UsableArcAngle = 360.f;
+	
+	/** The world rotation of the plane the ellipse is on */
+	UPROPERTY(EditAnywhere, Category = EllipticalPathPattern)
+	FRotator EllipsePlaneRotation = FRotator::ZeroRotator;
+};
\ No newline at end of file
