Chaos Mover - Fix movement on moving surfaces

diff --git a/Engine/Plugins/Experimental/ChaosMover/Source/ChaosMover/Private/ChaosMoverSimulation.cpp b/Engine/Plugins/Experimental/ChaosMover/Source/ChaosMover/Private/ChaosMoverSimulation.cpp
--- a/Engine/Plugins/Experimental/ChaosMover/Source/ChaosMover/Private/ChaosMoverSimulation.cpp
+++ b/Engine/Plugins/Experimental/ChaosMover/Source/ChaosMover/Private/ChaosMoverSimulation.cpp
@@ -12,6 +12,7 @@
 #include "ChaosMover/ChaosMovementModeTransition.h"
 #include "ChaosMover/ChaosMoverLog.h"
 #include "ChaosMover/ChaosMoverSimulationTypes.h"
+#include "ChaosMover/Utilities/ChaosGroundMovementUtils.h"
 #include "ChaosMover/Utilities/ChaosMoverQueryUtils.h"
 #include "ChaosMover/PathedMovement/ChaosPathedMovementTypes.h"
 #include "Framework/Threading.h"
@@ -277,6 +278,36 @@ void UChaosMoverSimulation::PreSimulationTickCharacter(const IChaosCharacterMove
 		StateMachine.QueueNextMode(CharacterDefaultInputs.SuggestedMovementMode);
 		CharacterDefaultInputs.SuggestedMovementMode = NAME_None;
 	}
+
+	// The state machine works in local space relative to the ground, so if we're on
+	// the ground make the state relative to the ground
+	// TODO - Can we use the local space in the sync state rather than subtracting and
+	//        re-adding ground velocity to the world space velocity?
+	FVector LocalGroundVelocity = FVector::ZeroVector;
+	const UBaseMovementMode* Mode = StateMachine.FindMovementMode(CurrentSyncState.MovementMode).Get();
+	if (Mode && Mode->HasGameplayTag(Mover_IsOnGround, true))
+	{
+		FMoverDefaultSyncState& PreSimDefaultSyncState = CurrentSyncState.SyncStateCollection.FindOrAddMutableDataByType<FMoverDefaultSyncState>();
+
+		FFloorCheckResult LastFloorResult;
+		if (Blackboard->TryGet(CommonBlackboard::LastFloorResult, LastFloorResult))
+		{
+			LocalGroundVelocity = FVector::VectorPlaneProject(
+				UChaosGroundMovementUtils::ComputeLocalGroundVelocity_Internal(PreSimDefaultSyncState.GetLocation_WorldSpace(), LastFloorResult),
+				LastFloorResult.HitResult.ImpactNormal);
+			
+			if (!LocalGroundVelocity.IsNearlyZero())
+			{
+				PreSimDefaultSyncState.SetTransforms_WorldSpace(
+					PreSimDefaultSyncState.GetLocation_WorldSpace(),
+					PreSimDefaultSyncState.GetOrientation_WorldSpace(),
+					PreSimDefaultSyncState.GetVelocity_WorldSpace() - LocalGroundVelocity);
+			}
+		}
+	}
+
+	FChaosMoverGroundSimState& GroundSimState = InternalSimData.FindOrAddMutableDataByType<FChaosMoverGroundSimState>();
+	GroundSimState.LocalVelocity = LocalGroundVelocity;
 }
 
 void UChaosMoverSimulation::OnPreSimulationTick(const FMoverTimeStep& TimeStep, const UE::ChaosMover::FSimulationInputData& InputData)
@@ -293,6 +324,7 @@ void UChaosMoverSimulation::OnPreSimulationTick(const FMoverTimeStep& TimeStep,
 
 	// Update the sync state from the current physics state
 	FMoverDefaultSyncState& PreSimDefaultSyncState = CurrentSyncState.SyncStateCollection.FindOrAddMutableDataByType<FMoverDefaultSyncState>();
+	PreSimDefaultSyncState.SetMovementBase(nullptr);
 	if (const Chaos::FPBDRigidParticleHandle* ParticleHandle = GetControlledParticle())
 	{
 		PreSimDefaultSyncState.SetTransforms_WorldSpace(ParticleHandle->GetX(), FRotator(ParticleHandle->GetR()), ParticleHandle->GetV());
@@ -332,7 +364,13 @@ void UChaosMoverSimulation::PostSimulationTickCharacter(const IChaosCharacterMov
 	{
 		// Linear motion
 		CharacterSimState.TargetDeltaPosition = PostSimDefaultSyncState->GetLocation_WorldSpace() - ParticleHandle->GetX();
-		ParticleHandle->SetV(PostSimDefaultSyncState->GetVelocity_WorldSpace());
+
+		FVector LocalGroundVelocity = FVector::ZeroVector;
+		if (const FChaosMoverGroundSimState* GroundSimState = InternalSimData.FindDataByType<FChaosMoverGroundSimState>())
+		{
+			LocalGroundVelocity = GroundSimState->LocalVelocity;
+		}
+		ParticleHandle->SetV(PostSimDefaultSyncState->GetVelocity_WorldSpace() + LocalGroundVelocity);
 
 		// Angular motion
 		FQuat TgtQuat = PostSimDefaultSyncState->GetOrientation_WorldSpace().Quaternion();
diff --git a/Engine/Plugins/Experimental/ChaosMover/Source/ChaosMover/Private/Character/Modes/ChaosWalkingMode.cpp b/Engine/Plugins/Experimental/ChaosMover/Source/ChaosMover/Private/Character/Modes/ChaosWalkingMode.cpp
--- a/Engine/Plugins/Experimental/ChaosMover/Source/ChaosMover/Private/Character/Modes/ChaosWalkingMode.cpp
+++ b/Engine/Plugins/Experimental/ChaosMover/Source/ChaosMover/Private/Character/Modes/ChaosWalkingMode.cpp
@@ -73,11 +73,9 @@ void UChaosWalkingMode::GenerateMove_Implementation(const FMoverTickStartData& S
 	UMoverBlackboard* SimBlackboard = Simulation->GetBlackboard_Mutable();
 	FFloorCheckResult LastFloorResult;
 	FVector MovementNormal;
-	FVector LocalGroundVelocity = FVector::ZeroVector;
 	if (SimBlackboard && SimBlackboard->TryGet(CommonBlackboard::LastFloorResult, LastFloorResult) && LastFloorResult.IsWalkableFloor())
 	{
 		MovementNormal = LastFloorResult.HitResult.ImpactNormal;
-		LocalGroundVelocity = UChaosGroundMovementUtils::ComputeLocalGroundVelocity_Internal(StartingSyncState->GetLocation_WorldSpace(), LastFloorResult);
 	}
 	else
 	{
@@ -114,10 +112,10 @@ void UChaosWalkingMode::GenerateMove_Implementation(const FMoverTickStartData& S
 		Params.MoveInput = FVector::ZeroVector;
 	}
 
-	const FVector CurrentRelativeVelocity = StartingSyncState->GetVelocity_WorldSpace() - LocalGroundVelocity;
+	const FVector CurrentVelocity = StartingSyncState->GetVelocity_WorldSpace();
 
 	Params.OrientationIntent = IntendedOrientation_WorldSpace;
-	Params.PriorVelocity = FVector::VectorPlaneProject(CurrentRelativeVelocity, MovementNormal);
+	Params.PriorVelocity = FVector::VectorPlaneProject(CurrentVelocity, MovementNormal);
 	Params.PriorOrientation = StartingSyncState->GetOrientation_WorldSpace();
 	Params.GroundNormal = MovementNormal;
 	Params.TurningRate = SharedSettings->TurningRate;
@@ -142,9 +140,6 @@ void UChaosWalkingMode::GenerateMove_Implementation(const FMoverTickStartData& S
 
 	OutProposedMove = UGroundMovementUtils::ComputeControlledGroundMove(Params);
 
-	// Add back in ground velocity to make the proposed move velocity in world space
-	OutProposedMove.LinearVelocity += LocalGroundVelocity;
-
 	if (SimBlackboard)
 	{
 		// Update the floor result and check the proposed move to prevent movement onto unwalkable surfaces
@@ -218,6 +213,8 @@ void UChaosWalkingMode::SimulationTick_Implementation(const FSimulationTickParam
 		constexpr float ParallelCosThreshold = 0.999f;
 		const bool bNonVerticalVelocity = !FVector::Parallel(TargetVelocity.GetSafeNormal(), UpDirection, ParallelCosThreshold);
 		const bool bUseProposedMove = bNonVerticalVelocity || ProposedMove.bHasDirIntent;
+		bool bNormalVelocityIntent = false;
+
 		if (bUseProposedMove)
 		{
 			const FVector ProposedMovePlaneVelocity = ProposedMove.LinearVelocity - ProposedMove.LinearVelocity.ProjectOnToNormal(StartingGroundNormal);
@@ -227,7 +224,12 @@ void UChaosWalkingMode::SimulationTick_Implementation(const FSimulationTickParam
 			FVector ProposedNormalVelocity = ProposedMove.LinearVelocity - ProposedMovePlaneVelocity;
 			if (ProposedNormalVelocity.SizeSquared() > UE_KINDA_SMALL_NUMBER)
 			{
-				TargetVelocity += ProposedNormalVelocity - TargetVelocity.ProjectOnToNormal(StartingGroundNormal);
+				TargetVelocity = ProposedMove.LinearVelocity;
+				bNormalVelocityIntent = true;
+			}
+			else
+			{
+				TargetVelocity = ProposedMovePlaneVelocity + TargetVelocity.ProjectOnToNormal(StartingGroundNormal);
 			}
 
 			TargetPosition += ProposedMovePlaneVelocity * DeltaSeconds;
@@ -242,10 +244,10 @@ void UChaosWalkingMode::SimulationTick_Implementation(const FSimulationTickParam
 			ProjectedGroundVelocity += DefaultSimInputs->PhysicsObjectGravity * DefaultSimInputs->UpDir * DeltaSeconds;
 		}
 		const bool bIsGroundMoving = ProjectedGroundVelocity.SizeSquared() > UE_KINDA_SMALL_NUMBER;
-		const FVector ProjectedRelativeVelocity = ProjectedVelocity - ProjectedGroundVelocity;
-		const float ProjectedRelativeNormalVelocity = FloorResult.HitResult.ImpactNormal.Dot(ProjectedVelocity - ProjectedGroundVelocity);
-		const float ProjectedRelativeVerticalVelocity = UpDirection.Dot(ProjectedVelocity - ProjectedGroundVelocity);
-		const float VerticalVelocityLimit = 2.0f / DeltaSeconds;
+		const FVector ProjectedRelativeVelocity = TargetVelocity - ProjectedGroundVelocity;
+		const float ProjectedRelativeNormalVelocity = FloorResult.HitResult.ImpactNormal.Dot(TargetVelocity - ProjectedGroundVelocity);
+		const float ProjectedRelativeVerticalVelocity = UpDirection.Dot(TargetVelocity - ProjectedGroundVelocity);
+		const float VerticalVelocityLimit = bNormalVelocityIntent ? 2.0f / DeltaSeconds : FMath::Abs(DefaultSimInputs->UpDir.Dot(DefaultSimInputs->Gravity) * DeltaSeconds);
 
 		bool bIsLiftingOffSurface = false;
 		if ((ProjectedRelativeNormalVelocity > VerticalVelocityLimit) && bIsGroundMoving && (ProjectedRelativeVerticalVelocity > VerticalVelocityLimit))
@@ -298,6 +300,8 @@ void UChaosWalkingMode::GetFloorAndCheckMovement(const FMoverDefaultSyncState& S
 	FVector DeltaPos = ProposedMove.LinearVelocity * DeltaSeconds;
 	OutDeltaPos = DeltaPos;
 
+	const float VelocityPadding = FMath::Max(DefaultSimInputs.UpDir.Dot(SyncState.GetVelocity_WorldSpace()) * DeltaSeconds, 0.0f);
+
 	UE::ChaosMover::Utils::FFloorSweepParams SweepParams{
 		.ResponseParams = DefaultSimInputs.CollisionResponseParams,
 		.QueryParams = DefaultSimInputs.CollisionQueryParams,
@@ -305,7 +309,7 @@ void UChaosWalkingMode::GetFloorAndCheckMovement(const FMoverDefaultSyncState& S
 		.DeltaPos = DeltaPos,
 		.UpDir = DefaultSimInputs.UpDir,
 		.World = DefaultSimInputs.World,
-		.QueryDistance = GetTargetHeight() + SharedSettings->MaxStepHeight,
+		.QueryDistance = GetTargetHeight() + SharedSettings->MaxStepHeight + VelocityPadding,
 		.QueryRadius = FMath::Min(GetGroundQueryRadius(), FMath::Max(DefaultSimInputs.PawnCollisionRadius - 5.0f, 0.0f)),
 		.MaxWalkSlopeCosine = GetMaxWalkSlopeCosine(),
 		.TargetHeight = GetTargetHeight(),
diff --git a/Engine/Plugins/Experimental/ChaosMover/Source/ChaosMover/Public/ChaosMover/ChaosMoverSimulationTypes.h b/Engine/Plugins/Experimental/ChaosMover/Source/ChaosMover/Public/ChaosMover/ChaosMoverSimulationTypes.h
--- a/Engine/Plugins/Experimental/ChaosMover/Source/ChaosMover/Public/ChaosMover/ChaosMoverSimulationTypes.h
+++ b/Engine/Plugins/Experimental/ChaosMover/Source/ChaosMover/Public/ChaosMover/ChaosMoverSimulationTypes.h
@@ -74,6 +74,28 @@ struct CHAOSMOVER_API FChaosMoverCharacterSimState : public FMoverDataStructBase
 	float TargetDeltaFacing = 0.0f;
 };
 
+// Default chaos mover character simulation state, contains state basic for character simulation
+USTRUCT(BlueprintType)
+struct CHAOSMOVER_API FChaosMoverGroundSimState : public FMoverDataStructBase
+{
+	GENERATED_BODY()
+
+	/** return newly allocated copy of this FMoverDataStructBase. MUST be overridden by derived types. */
+	virtual FMoverDataStructBase* Clone() const override
+	{
+		return new FChaosMoverGroundSimState(*this);
+	}
+
+	/** Gets the type info of this FMoverDataStructBase. MUST be overridden by derived types. */
+	virtual UScriptStruct* GetScriptStruct() const override
+	{
+		return StaticStruct();
+	}
+
+	UPROPERTY(VisibleAnywhere, BlueprintReadWrite, Category = ChaosMover)
+	FVector LocalVelocity = FVector::ZeroVector;
+};
+
 // Data block containing all default inputs required by the Chaos Mover simulation
 USTRUCT(BlueprintType)
 struct CHAOSMOVER_API FChaosMoverSimulationDefaultInputs : public FMoverDataStructBase
