ChaosMover: Fix to pathed movement resims

diff --git a/Engine/Plugins/Experimental/ChaosMover/Source/ChaosMover/Private/Backends/ChaosMoverBackend.cpp b/Engine/Plugins/Experimental/ChaosMover/Source/ChaosMover/Private/Backends/ChaosMoverBackend.cpp
--- a/Engine/Plugins/Experimental/ChaosMover/Source/ChaosMover/Private/Backends/ChaosMoverBackend.cpp
+++ b/Engine/Plugins/Experimental/ChaosMover/Source/ChaosMover/Private/Backends/ChaosMoverBackend.cpp
@@ -69,6 +69,7 @@ void UChaosMoverBackendComponent::InitializeComponent()
 	{
 		NullMovementMode = NewObject<UNullMovementMode>(&GetMoverComponent(), TEXT("NullMovementMode"));
 		ImmediateModeTransition = NewObject<UImmediateMovementModeTransition>(&GetMoverComponent(), TEXT("ImmediateModeTransition"));
+		TransformOnInit = GetMoverComponent().GetUpdatedComponentTransform();
 
 		// Create NetworkPhysicsComponent
 		if ((World->GetNetMode() != ENetMode::NM_Standalone) && Chaos::FPhysicsSolverBase::IsNetworkPhysicsPredictionEnabled())
@@ -300,6 +301,7 @@ void UChaosMoverBackendComponent::InitSimulation()
 	Params.CharacterConstraintProxy = CharacterGroundConstraint ? CharacterGroundConstraint->GetProxy<Chaos::FCharacterGroundConstraintProxy>() : nullptr;
 	Params.ActuationConstraintProxy = ActuationConstraintHandle.IsValid() ? ActuationConstraintHandle->GetProxy<Chaos::FJointConstraintPhysicsProxy>() : nullptr;
 	Params.ActuationConstraintKinematicEndPointProxy = Params.ActuationConstraintProxy && Params.ActuationConstraintProxy->GetConstraint() ? Params.ActuationConstraintProxy->GetConstraint()->GetKinematicEndPoint() : nullptr;
+	Params.TransformOnInit = TransformOnInit;
 	Params.PhysicsObject = GetPhysicsObject();
 	Params.Solver = GetPhysicsSolver();
 	Params.World = GetWorld();
diff --git a/Engine/Plugins/Experimental/ChaosMover/Source/ChaosMover/Private/ChaosMoverSimulation.cpp b/Engine/Plugins/Experimental/ChaosMover/Source/ChaosMover/Private/ChaosMoverSimulation.cpp
--- a/Engine/Plugins/Experimental/ChaosMover/Source/ChaosMover/Private/ChaosMoverSimulation.cpp
+++ b/Engine/Plugins/Experimental/ChaosMover/Source/ChaosMover/Private/ChaosMoverSimulation.cpp
@@ -107,16 +107,8 @@ void UChaosMoverSimulation::Init(const FInitParams& InitParams)
 
 	// Set the pathed movement basis to the initial location of the controlled particle
 	// We set this up before state machine init so the pathed movement modes have a pathed movement basis to work with on registration
-	if (PhysicsObject)
-	{
-		if (Chaos::FGeometryParticle* Particle = FPhysicsObjectExternalInterface::LockRead(PhysicsObject)->GetRigidParticle(PhysicsObject))
-		{
-			FTransform NewPathBasis;
-			NewPathBasis.SetTranslation(Particle->GetX());
-			NewPathBasis.SetRotation(Particle->GetR());
-			SetPathedMovementBasis(NewPathBasis);
-		}
-	}
+	FTransform NewPathBasis(InitParams.TransformOnInit.GetRotation(), InitParams.TransformOnInit.GetLocation());
+	SetPathedMovementBasis(NewPathBasis);
 
 	UE::ChaosMover::FMoverStateMachine::FInitParams StateMachineInitParams;
 	StateMachineInitParams.ImmediateMovementModeTransition = InitParams.ImmediateModeTransition;
@@ -665,6 +657,9 @@ void UChaosMoverSimulation::TraceMoverData(const UE::ChaosMover::FSimulationOutp
 		int32 SolverID = CVD_TRACE_GET_SOLVER_ID_FROM_WORLD(World);
 
 		UE::MoverUtils::FMoverCVDRuntimeTrace::TraceMoverData(SolverID, ParticleID, &OutputData.LastUsedInputCmd, &OutputData.SyncState, &LocalSimDataCollections);
+
+		// Clear debug sim data to avoid recording past debug data as current on the next frame
+		DebugSimData.Empty();
 	}
 #endif
 }
@@ -753,7 +748,7 @@ void UChaosMoverSimulation::OnMovementModeChanged(const FMovementModeChangedEven
 					IChaosPathedMovementModeInterface* NextPathedMovementMode = Cast<IChaosPathedMovementModeInterface>(NextMovementActuationInterface);
 					if (NextPathedMovementMode)
 					{	
-						const float PathProgress = PresimPathedMovementState ? PresimPathedMovementState->CurrentPathProgress : 0.0f;
+						const float PathProgress = PresimPathedMovementState ? PresimPathedMovementState->LastChangePathProgress : 0.0f;
 						// We resume the path at the reset progress
 						// (if not 0, CurrentPathProgress was either not reset properly when the movement mode changed, or deliberately set to a different value)
 						TargetLastFrame = NextPathedMovementMode->CalcTargetRelativeTransform(PathProgress);
diff --git a/Engine/Plugins/Experimental/ChaosMover/Source/ChaosMover/Private/PathedMovement/ChaosPathedMovementMode.cpp b/Engine/Plugins/Experimental/ChaosMover/Source/ChaosMover/Private/PathedMovement/ChaosPathedMovementMode.cpp
--- a/Engine/Plugins/Experimental/ChaosMover/Source/ChaosMover/Private/PathedMovement/ChaosPathedMovementMode.cpp
+++ b/Engine/Plugins/Experimental/ChaosMover/Source/ChaosMover/Private/PathedMovement/ChaosPathedMovementMode.cpp
@@ -10,6 +10,7 @@
 #include "Framework/Threading.h"
 #include "Kismet/KismetMathLibrary.h"
 #include "Misc/DataValidation.h"
+#include "ChaosVisualDebugger/ChaosVisualDebuggerTrace.h"
 
 #if WITH_EDITOR
 #include "Engine/SimpleConstructionScript.h"
@@ -58,9 +59,12 @@ void UChaosPathedMovementMode::SimulationTick_Implementation(const FSimulationTi
 {
 	Chaos::EnsureIsInPhysicsThreadContext();
 
-#if !UE_BUILD_SHIPPING
-	DebugData.Advance();
-#endif // !UE_BUILD_SHIPPING
+#if WITH_CHAOS_VISUAL_DEBUGGER
+	if (FChaosVisualDebuggerTrace::IsTracing())
+	{
+		DebugData.Advance();
+	}
+#endif // WITH_CHAOS_VISUAL_DEBUGGER
 
 	// We retrieve the previous properties so we can compare with the new. If there are none we default to FChaosPathedMovementState()
 	const FChaosPathedMovementState* FoundPreviousMoveState = Params.StartState.SyncState.SyncStateCollection.FindDataByType<FChaosPathedMovementState>();
@@ -75,6 +79,17 @@ void UChaosPathedMovementMode::SimulationTick_Implementation(const FSimulationTi
 	// We initialize the output movement state from the previous movement state
 	OutputPathedMovemenState = PreviousMoveState;
 
+	// Compute the current path progress before we modify movement properties or update LastChangeFrame 
+	float CurrentPathProgress = PreviousMoveState.LastChangePathProgress;
+	if (PreviousProperties.bWantsToPlay && !PreviousMoveState.bHasFinished)
+	{
+		const float FrameDuration = Params.TimeStep.StepMs * 0.001f;
+		const int32 PreviousLastChangeFrame = (PreviousMoveState.LastChangeFrame != INDEX_NONE) ? PreviousMoveState.LastChangeFrame : 0;
+		const int32 FramesSinceLastChange = Params.TimeStep.ServerFrame - PreviousMoveState.LastChangeFrame;
+		const float PlaybackSpeed = (OneWayPlaybackDuration >= 0.0f) ? (1.0f / OneWayPlaybackDuration) : 0.0f;
+		const float PreviousDirection = PreviousMoveState.bIsPathProgressionIncreasing ? 1.0f : -1.0f;
+		CurrentPathProgress += FramesSinceLastChange * (PreviousDirection * FrameDuration * PlaybackSpeed);
+	}
 	
 	bool bReceivedNewStartRequest = false;
 
@@ -93,6 +108,12 @@ void UChaosPathedMovementMode::SimulationTick_Implementation(const FSimulationTi
 
 		// ... then the input properties enter into effect
 		OutputPathedMovemenState.PropertiesInEffect = PathedMovementInputs->Props;
+
+		// Since we changed movement properties, we record the path progress and frame at which it changed
+		// This is all that should be needed to be able to compute path progress for future frames, until movement properties change again
+		// From now on all frame differences should be using these values and not the ones in PreviousMoveState
+		OutputPathedMovemenState.LastChangePathProgress = CurrentPathProgress;
+		OutputPathedMovemenState.LastChangeFrame = Params.TimeStep.ServerFrame;
 	}
 
 	// PropertiesInEffect is const here, to indicate it is a bad idea to modify those, since they are used to detect external changes in mutable properties,
@@ -120,9 +141,9 @@ void UChaosPathedMovementMode::SimulationTick_Implementation(const FSimulationTi
 		OutputPathedMovemenState.bIsPathProgressionIncreasing = !CurrentProperties.bWantsReversePlayback;
 		if (CurrentProperties.bWantsOneWayPlayback)
 		{
-			if (OutputPathedMovemenState.CurrentPathProgress == (OutputPathedMovemenState.bIsPathProgressionIncreasing ? 1.0f : 0.0f))
+			if (CurrentPathProgress == (OutputPathedMovemenState.bIsPathProgressionIncreasing ? 1.0f : 0.0f))
 			{
-				OutputPathedMovemenState.CurrentPathProgress = 1.0f - OutputPathedMovemenState.CurrentPathProgress;
+				CurrentPathProgress = 1.0f - CurrentPathProgress;
 			}
 		}
 	}
@@ -134,13 +155,14 @@ void UChaosPathedMovementMode::SimulationTick_Implementation(const FSimulationTi
 		bJustRestarted = true;
 	}
 
+	float NewPathProgress = CurrentPathProgress;
 	if (CurrentProperties.bWantsToPlay && (!PreviousMoveState.bHasFinished || bReceivedNewStartRequest))
 	{
 		const float FrameDuration = Params.TimeStep.StepMs * 0.001f;
 		const float PlaybackSpeed = (OneWayPlaybackDuration >= 0.0f) ? (1.0f / OneWayPlaybackDuration) : 0.0f;
 		const float Direction = OutputPathedMovemenState.bIsPathProgressionIncreasing ? 1.0f : -1.0f;
-		const float DeltaPathProgress = Direction * PlaybackSpeed * FrameDuration;
-		const float NewPathProgress = OutputPathedMovemenState.CurrentPathProgress + DeltaPathProgress;
+		const int32 NumFramesSinceLastChange = (Params.TimeStep.ServerFrame - OutputPathedMovemenState.LastChangeFrame) + 1;
+		NewPathProgress = OutputPathedMovemenState.LastChangePathProgress + NumFramesSinceLastChange * (Direction * PlaybackSpeed * FrameDuration);
 
 		if (CurrentProperties.bWantsOneWayPlayback)
 		{
@@ -151,20 +173,20 @@ void UChaosPathedMovementMode::SimulationTick_Implementation(const FSimulationTi
 				{
 					if (CurrentProperties.bWantsLoopingPlayback)
 					{
-						OutputPathedMovemenState.CurrentPathProgress = FMath::Fmod(NewPathProgress, 1.0f);
+						NewPathProgress = FMath::Fmod(NewPathProgress, 1.0f);
 						bJustRestarted = true;
+						OutputPathedMovemenState.LastChangePathProgress = NewPathProgress;
+						OutputPathedMovemenState.LastChangeFrame = Params.TimeStep.ServerFrame;
 					}
 					else
 					{
-						OutputPathedMovemenState.CurrentPathProgress = 1.0f;
-						OutputPathedMovemenState.bHasFinished = true;
+						NewPathProgress = 1.0f;
 						bHasJustFinished = true;
+						OutputPathedMovemenState.bHasFinished = true;
+						OutputPathedMovemenState.LastChangePathProgress = NewPathProgress;
+						OutputPathedMovemenState.LastChangeFrame = Params.TimeStep.ServerFrame;
 					}
 				}
-				else
-				{
-					OutputPathedMovemenState.CurrentPathProgress = NewPathProgress;
-				}
 			}
 			else // One way reverse playback
 			{
@@ -172,20 +194,20 @@ void UChaosPathedMovementMode::SimulationTick_Implementation(const FSimulationTi
 				{
 					if (CurrentProperties.bWantsLoopingPlayback)
 					{
-						OutputPathedMovemenState.CurrentPathProgress = 1.0f - FMath::Fmod(-NewPathProgress, 1.0f);
+						NewPathProgress = 1.0f - FMath::Fmod(-NewPathProgress, 1.0f);
 						bJustRestarted = true;
+						OutputPathedMovemenState.LastChangePathProgress = NewPathProgress;
+						OutputPathedMovemenState.LastChangeFrame = Params.TimeStep.ServerFrame;
 					}
 					else
 					{
-						OutputPathedMovemenState.CurrentPathProgress = 0.0f;
-						OutputPathedMovemenState.bHasFinished = true;
+						NewPathProgress = 0.0f;
 						bHasJustFinished = true;
+						OutputPathedMovemenState.bHasFinished = true;
+						OutputPathedMovemenState.LastChangePathProgress = NewPathProgress;
+						OutputPathedMovemenState.LastChangeFrame = Params.TimeStep.ServerFrame;
 					}
 				}
-				else
-				{
-					OutputPathedMovemenState.CurrentPathProgress = NewPathProgress;
-				}
 			}
 		}
 		else
@@ -196,35 +218,39 @@ void UChaosPathedMovementMode::SimulationTick_Implementation(const FSimulationTi
 				if (!CurrentProperties.bWantsReversePlayback || CurrentProperties.bWantsLoopingPlayback)
 				{
 					OutputPathedMovemenState.bIsPathProgressionIncreasing = false;
-					OutputPathedMovemenState.CurrentPathProgress = 1.0f - FMath::Fmod(NewPathProgress, 1.0f);
+					NewPathProgress = 1.0f - FMath::Fmod(NewPathProgress, 1.0f);
 					bHasJustBounced = true;
+					OutputPathedMovemenState.LastChangePathProgress = NewPathProgress;
+					OutputPathedMovemenState.LastChangeFrame = Params.TimeStep.ServerFrame;
 				}
 				else
 				{
-					OutputPathedMovemenState.CurrentPathProgress = 1.0f;
-					OutputPathedMovemenState.bHasFinished = true;
+					NewPathProgress = 1.0f;
 					bHasJustFinished = true;
+					OutputPathedMovemenState.bHasFinished = true;
+					OutputPathedMovemenState.LastChangePathProgress = NewPathProgress;
+					OutputPathedMovemenState.LastChangeFrame = Params.TimeStep.ServerFrame;
 				}
 			}
 			else if (!OutputPathedMovemenState.bIsPathProgressionIncreasing && NewPathProgress <= 0.0f)
 			{
 				if (CurrentProperties.bWantsReversePlayback || CurrentProperties.bWantsLoopingPlayback)
 				{
-					OutputPathedMovemenState.bIsPathProgressionIncreasing = true;
-					OutputPathedMovemenState.CurrentPathProgress = FMath::Fmod(-NewPathProgress, 1.0f);
+					NewPathProgress = FMath::Fmod(-NewPathProgress, 1.0f);
 					bHasJustBounced = true;
+					OutputPathedMovemenState.bIsPathProgressionIncreasing = true;
+					OutputPathedMovemenState.LastChangePathProgress = NewPathProgress;
+					OutputPathedMovemenState.LastChangeFrame = Params.TimeStep.ServerFrame;
 				}
 				else
 				{
-					OutputPathedMovemenState.CurrentPathProgress = 0.0f;
-					OutputPathedMovemenState.bHasFinished = true;
+					NewPathProgress = 0.0f;
 					bHasJustFinished = true;
+					OutputPathedMovemenState.bHasFinished = true;
+					OutputPathedMovemenState.LastChangePathProgress = NewPathProgress;
+					OutputPathedMovemenState.LastChangeFrame = Params.TimeStep.ServerFrame;
 				}
 			}
-			else
-			{
-				OutputPathedMovemenState.CurrentPathProgress = NewPathProgress;
-			}
 		}
 	}
 
@@ -258,7 +284,7 @@ void UChaosPathedMovementMode::SimulationTick_Implementation(const FSimulationTi
 		OnPlaybackBounced(Params.TimeStep.BaseSimTimeMs);
 	}
 
-	const Chaos::FRigidTransform3 TargetRelativeTransform = CalcTargetRelativeTransform(OutputPathedMovemenState.CurrentPathProgress);
+	const Chaos::FRigidTransform3 TargetRelativeTransform = CalcTargetRelativeTransform(NewPathProgress);
 
 	FMoverDefaultSyncState& OutputDefaultSyncState = OutputState.SyncState.SyncStateCollection.FindOrAddMutableDataByType<FMoverDefaultSyncState>();
 	
@@ -269,14 +295,22 @@ void UChaosPathedMovementMode::SimulationTick_Implementation(const FSimulationTi
 											  FVector::ZeroVector, // Fix this eventually to return the velocity tangential to the path, including due to a change in path basis since last update
 											  nullptr ); // no movement base
 
-#if !UE_BUILD_SHIPPING
-	DebugData.ServerFrame = Params.TimeStep.ServerFrame;
-	DebugData.PreSimMoveState = PreviousMoveState;
-	DebugData.PostSimMoveState = OutputPathedMovemenState;
-	DebugData.FinalWorldTranslation = WorldTranslation;
-	DebugData.FinalWorldRotation = WorldRotation;
-	// Record state for debugging
-	UE::ChaosMover::GetDebugSimData(Simulation).FindOrAddMutableDataByType<FChaosPathedMovementModeDebugData>() = DebugData;
+#if WITH_CHAOS_VISUAL_DEBUGGER
+	if (FChaosVisualDebuggerTrace::IsTracing())
+	{
+		DebugData.ServerFrame = Params.TimeStep.ServerFrame;
+		DebugData.PreSimMoveState = PreviousMoveState;
+		DebugData.PostSimMoveState = OutputPathedMovemenState;
+		const FMoverDefaultSyncState* PreviousDefaultSyncState = Params.StartState.SyncState.SyncStateCollection.FindDataByType<FMoverDefaultSyncState>();
+		DebugData.PreviousWorldTranslation = PreviousDefaultSyncState ? PreviousDefaultSyncState->GetTransform_WorldSpace().GetLocation() : FVector::ZeroVector;
+		DebugData.PreviousWorldRotation = PreviousDefaultSyncState ? FRotator(PreviousDefaultSyncState->GetTransform_WorldSpace().GetRotation()) : FRotator::ZeroRotator;
+		DebugData.NewRelativeTranslation = TargetRelativeTransform.GetTranslation();
+		DebugData.NewRelativeRotation = FRotator(TargetRelativeTransform.GetRotation());
+		DebugData.NewWorldTranslation = WorldTranslation;
+		DebugData.NewWorldRotation = WorldRotation;
+		// Record state for debugging
+		UE::ChaosMover::GetDebugSimData(Simulation).FindOrAddMutableDataByType<FChaosPathedMovementModeDebugData>() = DebugData;
+	}
 #endif
 }
 
@@ -287,9 +321,12 @@ void UChaosPathedMovementMode::OnPlaybackStarting(float BaseSimTimeMs)
 	{
 		Simulation->AddEvent(MakeShared<FPathedMovementStartedEventData>(BaseSimTimeMs));
 
-#if !UE_BUILD_SHIPPING
+#if WITH_CHAOS_VISUAL_DEBUGGER
+	if (FChaosVisualDebuggerTrace::IsTracing())
+	{
 		DebugData.bPlaybackStartedThisFrame = true;
-#endif // !UE_BUILD_SHIPPING
+	}
+#endif // WITH_CHAOS_VISUAL_DEBUGGER
 	}
 }
 
@@ -300,10 +337,13 @@ void UChaosPathedMovementMode::OnPlaybackStopping(float BaseSimTimeMs, bool bRea
 	{
 		Simulation->AddEvent(MakeShared<FPathedMovementStoppedEventData>(BaseSimTimeMs, bReachedEndOfPlayback));
 
-#if !UE_BUILD_SHIPPING
+#if WITH_CHAOS_VISUAL_DEBUGGER
+	if (FChaosVisualDebuggerTrace::IsTracing())
+	{
 		DebugData.bPlaybackStoppedThisFrame = true;
 		DebugData.bPlaybackReachedEndThisFrame = bReachedEndOfPlayback;
-#endif // !UE_BUILD_SHIPPING
+	}
+#endif // WITH_CHAOS_VISUAL_DEBUGGER
 	}
 }
 
@@ -314,9 +354,12 @@ void UChaosPathedMovementMode::OnPlaybackBounced(float BaseSimTimeMs)
 	{
 		Simulation->AddEvent(MakeShared<FPathedMovementBouncedEventData>(BaseSimTimeMs));
 
-#if !UE_BUILD_SHIPPING
-		DebugData.bPlaybackBouncedThisFrame = true;
-#endif // !UE_BUILD_SHIPPING
+#if WITH_CHAOS_VISUAL_DEBUGGER
+		if (FChaosVisualDebuggerTrace::IsTracing())
+		{
+			DebugData.bPlaybackBouncedThisFrame = true;
+		}
+#endif // WITH_CHAOS_VISUAL_DEBUGGER
 	}
 }
 
@@ -327,9 +370,12 @@ void UChaosPathedMovementMode::OnInReverseChanging(float BaseSimTimeMs, bool bNe
 	{
 		Simulation->AddEvent(MakeShared<FPathedMovementInReverseChangedEventData>(BaseSimTimeMs, bNewInReverse));
 
-#if !UE_BUILD_SHIPPING
-		DebugData.bPlaybackReverseChangedThisFrame = true;
-#endif // !UE_BUILD_SHIPPING
+#if WITH_CHAOS_VISUAL_DEBUGGER
+		if (FChaosVisualDebuggerTrace::IsTracing())
+		{
+			DebugData.bPlaybackReverseChangedThisFrame = true;
+		}
+#endif // WITH_CHAOS_VISUAL_DEBUGGER
 	}
 }
 
@@ -340,9 +386,12 @@ void UChaosPathedMovementMode::OnLoopingChanging(float BaseSimTimeMs, bool bNewL
 	{
 		Simulation->AddEvent(MakeShared<FPathedMovementLoopingChangedEventData>(BaseSimTimeMs, bNewLooping));
 
-#if !UE_BUILD_SHIPPING
-		DebugData.bPlaybackLoopingChangedThisFrame = true;
-#endif // !UE_BUILD_SHIPPING
+#if WITH_CHAOS_VISUAL_DEBUGGER
+		if (FChaosVisualDebuggerTrace::IsTracing())
+		{
+			DebugData.bPlaybackLoopingChangedThisFrame = true;
+		}
+#endif // WITH_CHAOS_VISUAL_DEBUGGER
 	}
 }
 
@@ -353,9 +402,12 @@ void UChaosPathedMovementMode::OnOneWayChanging(float BaseSimTimeMs, bool bNewOn
 	{
 		Simulation->AddEvent(MakeShared<FPathedMovementOneWayChangedEventData>(BaseSimTimeMs, bNewOneWay));
 
-#if !UE_BUILD_SHIPPING
-		DebugData.bPlaybackOneWayChangedThisFrame = true;
-#endif // !UE_BUILD_SHIPPING
+#if WITH_CHAOS_VISUAL_DEBUGGER
+		if (FChaosVisualDebuggerTrace::IsTracing())
+		{
+			DebugData.bPlaybackOneWayChangedThisFrame = true;
+		}
+#endif // WITH_CHAOS_VISUAL_DEBUGGER
 	}
 }
 
diff --git a/Engine/Plugins/Experimental/ChaosMover/Source/ChaosMover/Private/PathedMovement/ChaosPathedMovementPatternBase.cpp b/Engine/Plugins/Experimental/ChaosMover/Source/ChaosMover/Private/PathedMovement/ChaosPathedMovementPatternBase.cpp
--- a/Engine/Plugins/Experimental/ChaosMover/Source/ChaosMover/Private/PathedMovement/ChaosPathedMovementPatternBase.cpp
+++ b/Engine/Plugins/Experimental/ChaosMover/Source/ChaosMover/Private/PathedMovement/ChaosPathedMovementPatternBase.cpp
@@ -68,7 +68,7 @@ float UChaosPathedMovementPatternBase::ConvertPathToPatternProgress(float Overal
 	const float PathProgressSinceStart = FMath::Min(EndAtPathProgress, OverallPathProgress) - StartAtPathProgress;
 	const float PathProgressPerPatternLoop = (EndAtPathProgress - StartAtPathProgress) / NumLoopsPerPath;
 	float CurLoopPathProgress = FMath::Fmod(PathProgressSinceStart, PathProgressPerPatternLoop);
-	if (CurLoopPathProgress == 0.f && PathProgressSinceStart > 0.f)
+	if (FMath::IsNearlyZero(CurLoopPathProgress) && PathProgressSinceStart > UE_SMALL_NUMBER)
 	{
 		// Treat path progress that matches the per-loop span exactly as 100%, not 0%
 		CurLoopPathProgress = PathProgressPerPatternLoop;
@@ -89,19 +89,31 @@ float UChaosPathedMovementPatternBase::ConvertPathToPatternProgress(float Overal
 
 	// Now convert from overall path progress to pattern progress
 	const float PatternProgress = CurLoopPathProgress / PathProgressPerPatternLoop;
-	const float OffsettedPatternProgress = FMath::Fmod(PatternProgress + LoopOffset, 1.0f);
+	const float OffsettedPatternProgress = PatternProgress + LoopOffset;
+	float FinalPatternProgress = FMath::Fmod(OffsettedPatternProgress, 1.0f);
+	if (FMath::IsNearlyZero(FinalPatternProgress) && OffsettedPatternProgress > UE_SMALL_NUMBER)
+	{
+		// Treat progress that matches one loop as 100%, not 0% ("modulo-inclusive Fmod") otherwise this Fmod could undo the fix above
+		FinalPatternProgress = 1.0f;
+	}
 
-	return OffsettedPatternProgress;
+	return FinalPatternProgress;
 }
 
-float UChaosPathedMovementPatternBase::ConvertPatternToPathProgress(float PatternProgress) const
+float UChaosPathedMovementPatternBase::ConvertPatternToPathProgress(float FinalPatternProgress) const
 {
 	if (NumLoopsPerPath <= 0 || StartAtPathProgress >= EndAtPathProgress)
 	{
 		return -1.0f;
 	}
 
-	const float OffsettedPatternProgress = FMath::Fmod(FMath::Clamp(PatternProgress, 0.0f, 1.0f) - LoopOffset, 1.0f);
+	const float PatternProgressOffsetRemoved = FMath::Clamp(FinalPatternProgress, 0.0f, 1.0f) - LoopOffset;
+	float PatternProgress = FMath::Fmod(PatternProgressOffsetRemoved, 1.0f);
+	if (FMath::IsNearlyZero(PatternProgress) && PatternProgressOffsetRemoved > UE_SMALL_NUMBER)
+	{
+		// Treat PatternProgress = 1.0f as 1.0f
+		PatternProgress = 1.0f;
+	}
 	const float PathProgressPerPatternLoop = (EndAtPathProgress - StartAtPathProgress) / NumLoopsPerPath;
 	float CurLoopPathProgress = PatternProgress * PathProgressPerPatternLoop;
 	if (!IsOneWay)
@@ -109,7 +121,7 @@ float UChaosPathedMovementPatternBase::ConvertPatternToPathProgress(float Patter
 		CurLoopPathProgress /= 2.f;
 	}
 
-	return StartAtPathProgress + (EndAtPathProgress - StartAtPathProgress) * CurLoopPathProgress;
+	return FMath::Clamp(StartAtPathProgress + (EndAtPathProgress - StartAtPathProgress) * CurLoopPathProgress, 0.0f, 1.0f);
 }
 
 FTransform UChaosPathedMovementPatternBase::CalcTargetRelativeTransform(float OverallPathProgress, const FTransform& CurTargetTransform) const
diff --git a/Engine/Plugins/Experimental/ChaosMover/Source/ChaosMover/Private/PathedMovement/ChaosPathedMovementTypes.cpp b/Engine/Plugins/Experimental/ChaosMover/Source/ChaosMover/Private/PathedMovement/ChaosPathedMovementTypes.cpp
--- a/Engine/Plugins/Experimental/ChaosMover/Source/ChaosMover/Private/PathedMovement/ChaosPathedMovementTypes.cpp
+++ b/Engine/Plugins/Experimental/ChaosMover/Source/ChaosMover/Private/PathedMovement/ChaosPathedMovementTypes.cpp
@@ -14,10 +14,10 @@ void FChaosMutablePathedMovementProperties::NetSerialize(FArchive& Ar)
 
 void FChaosMutablePathedMovementProperties::ToString(FAnsiStringBuilderBase& Out) const
 {
-	Out.Appendf("Wants to Play Path: %hs | \n", *LexToString(bWantsToPlay));
-	Out.Appendf("Wants Reverse Playback: %hs | \n", *LexToString(bWantsReversePlayback));
-	Out.Appendf("Wants Looping Playback: %hs | \n", *LexToString(bWantsLoopingPlayback));
-	Out.Appendf("Wants One Way Playback: %hs | \n", *LexToString(bWantsOneWayPlayback));
+	Out.Appendf("Wants to Play Path: %hs | ", *LexToString(bWantsToPlay));
+	Out.Appendf("Wants Reverse Playback: %hs | ", *LexToString(bWantsReversePlayback));
+	Out.Appendf("Wants Looping Playback: %hs | ", *LexToString(bWantsLoopingPlayback));
+	Out.Appendf("Wants One Way Playback: %hs\n", *LexToString(bWantsOneWayPlayback));
 }
 
 bool FChaosMutablePathedMovementProperties::operator==(const FChaosMutablePathedMovementProperties& Other) const
@@ -78,13 +78,15 @@ bool FChaosPathedMovementState::NetSerialize(FArchive& Ar, UPackageMap* Map, boo
 
 	static FChaosPathedMovementState Defaults;
 
-	bool HasCurrentPathProgress;
+	bool HasLastChangePathProgress;
 	bool HasLatestMovementStartFrame;
+	bool HasLastChangeFrame;
 
 	if (Ar.IsSaving())
 	{
-		HasCurrentPathProgress = CurrentPathProgress != Defaults.CurrentPathProgress;
+		HasLastChangePathProgress = LastChangePathProgress != Defaults.LastChangePathProgress;
 		HasLatestMovementStartFrame = LatestMovementStartFrame != Defaults.LatestMovementStartFrame;
+		HasLastChangeFrame = LastChangeFrame != Defaults.LastChangeFrame;
 	}
 	else
 	{
@@ -92,20 +94,25 @@ bool FChaosPathedMovementState::NetSerialize(FArchive& Ar, UPackageMap* Map, boo
 		*this = Defaults;
 	}
 
-	Ar.SerializeBits(&HasCurrentPathProgress, 1);
+	Ar.SerializeBits(&HasLastChangePathProgress, 1);
 	Ar.SerializeBits(&HasLatestMovementStartFrame, 1);
+	Ar.SerializeBits(&HasLastChangeFrame, 1);
 	Ar.SerializeBits(&bIsPathProgressionIncreasing, 1);
 	Ar.SerializeBits(&bHasFinished, 1);
 	PropertiesInEffect.NetSerialize(Ar);
 
-	if (HasCurrentPathProgress)
+	if (HasLastChangePathProgress)
 	{
-		Ar << CurrentPathProgress;
+		Ar << LastChangePathProgress;
 	}
 	if (HasLatestMovementStartFrame)
 	{
 		Ar << LatestMovementStartFrame;
 	}
+	if (HasLastChangeFrame)
+	{
+		Ar << LastChangeFrame;
+	}
 
 	return true;
 }
@@ -114,7 +121,8 @@ void FChaosPathedMovementState::ToString(FAnsiStringBuilderBase& Out) const
 {
 	FMoverDataStructBase::ToString(Out);
 
-	Out.Appendf("Current Path Progress: %.2f | ", CurrentPathProgress);
+	Out.Appendf("Last Change Path Progress: %.5f | ", LastChangePathProgress);
+	Out.Appendf("Last Change Frame: %d | ", LastChangeFrame);
 	Out.Appendf("Latest Movement Start Frame: %d | ", LatestMovementStartFrame);
 	Out.Appendf("Is Path Progress Increasing: %s | ", bIsPathProgressionIncreasing ? TEXT("True") : TEXT("False"));
 	Out.Appendf("Has Finished: %s | ", bHasFinished ? TEXT("True") : TEXT("False"));
@@ -125,13 +133,14 @@ bool FChaosPathedMovementState::ShouldReconcile(const FMoverDataStructBase& Auth
 {
 	const FChaosPathedMovementState& TypedAuthority = static_cast<const FChaosPathedMovementState&>(AuthorityState);
 
-	static constexpr float PathProgressTolerance = 0.05f;
+	static constexpr float PathProgressTolerance = 1e-4f;
 
 	if (PropertiesInEffect != TypedAuthority.PropertiesInEffect
 		|| LatestMovementStartFrame != TypedAuthority.LatestMovementStartFrame
+		|| LastChangeFrame != TypedAuthority.LastChangeFrame
 		|| bIsPathProgressionIncreasing != TypedAuthority.bIsPathProgressionIncreasing
 		|| bHasFinished != TypedAuthority.bHasFinished
-		|| !FMath::IsNearlyEqual(CurrentPathProgress, TypedAuthority.CurrentPathProgress, PathProgressTolerance))
+		|| !FMath::IsNearlyEqual(LastChangePathProgress, TypedAuthority.LastChangePathProgress, PathProgressTolerance))
 	{
 		return true;
 	}
@@ -146,7 +155,8 @@ void FChaosPathedMovementState::Interpolate(const FMoverDataStructBase& From, co
 
 	PropertiesInEffect = (Pct < .5f) ? TypedFrom.PropertiesInEffect : TypedTo.PropertiesInEffect;
 	LatestMovementStartFrame = (Pct < .5f) ? TypedFrom.LatestMovementStartFrame : TypedTo.LatestMovementStartFrame;
-	CurrentPathProgress = FMath::Lerp(TypedFrom.CurrentPathProgress, TypedTo.CurrentPathProgress, Pct);
+	LastChangeFrame = (Pct < .5f) ? TypedFrom.LastChangeFrame : TypedTo.LastChangeFrame;
+	LastChangePathProgress = FMath::Lerp(TypedFrom.LastChangePathProgress, TypedTo.LastChangePathProgress, Pct);
 	bIsPathProgressionIncreasing = (Pct < .5f) ? TypedFrom.bIsPathProgressionIncreasing : TypedTo.bIsPathProgressionIncreasing;
 	bHasFinished = (Pct < .5f) ? TypedFrom.bHasFinished : TypedTo.bHasFinished;
 }
diff --git a/Engine/Plugins/Experimental/ChaosMover/Source/ChaosMover/Public/ChaosMover/Backends/ChaosMoverBackend.h b/Engine/Plugins/Experimental/ChaosMover/Source/ChaosMover/Public/ChaosMover/Backends/ChaosMoverBackend.h
--- a/Engine/Plugins/Experimental/ChaosMover/Source/ChaosMover/Public/ChaosMover/Backends/ChaosMoverBackend.h
+++ b/Engine/Plugins/Experimental/ChaosMover/Source/ChaosMover/Public/ChaosMover/Backends/ChaosMoverBackend.h
@@ -97,6 +97,8 @@ protected:
 	FConstraintInstance ActuationConstraintInstance;
 	FPhysicsUserData ActuationConstraintPhysicsUserData;
 
+	FTransform TransformOnInit;
+
 	bool bIsUsingAsyncPhysics = false;
 	bool bWantsDestroySim = false;
 	bool bWantsCreateSim = false;
diff --git a/Engine/Plugins/Experimental/ChaosMover/Source/ChaosMover/Public/ChaosMover/ChaosMoverSimulation.h b/Engine/Plugins/Experimental/ChaosMover/Source/ChaosMover/Public/ChaosMover/ChaosMoverSimulation.h
--- a/Engine/Plugins/Experimental/ChaosMover/Source/ChaosMover/Public/ChaosMover/ChaosMoverSimulation.h
+++ b/Engine/Plugins/Experimental/ChaosMover/Source/ChaosMover/Public/ChaosMover/ChaosMoverSimulation.h
@@ -61,6 +61,7 @@ public:
 		TWeakObjectPtr<UNullMovementMode> NullMovementMode = nullptr;
 		TWeakObjectPtr<UImmediateMovementModeTransition> ImmediateModeTransition = nullptr;
 		TWeakObjectPtr<UMovementMixer> MovementMixer = nullptr;
+		FTransform TransformOnInit = FTransform::Identity;
 		Chaos::FCharacterGroundConstraintProxy* CharacterConstraintProxy = nullptr;
 		Chaos::FJointConstraintPhysicsProxy* ActuationConstraintProxy = nullptr;
 		Chaos::FSingleParticlePhysicsProxy* ActuationConstraintKinematicEndPointProxy = nullptr;
diff --git a/Engine/Plugins/Experimental/ChaosMover/Source/ChaosMover/Public/ChaosMover/PathedMovement/ChaosPathedMovementMode.h b/Engine/Plugins/Experimental/ChaosMover/Source/ChaosMover/Public/ChaosMover/PathedMovement/ChaosPathedMovementMode.h
--- a/Engine/Plugins/Experimental/ChaosMover/Source/ChaosMover/Public/ChaosMover/PathedMovement/ChaosPathedMovementMode.h
+++ b/Engine/Plugins/Experimental/ChaosMover/Source/ChaosMover/Public/ChaosMover/PathedMovement/ChaosPathedMovementMode.h
@@ -138,7 +138,7 @@ private:
 	
 	friend class UChaosPathedMovementDebugDrawComponent;
 
-#if !UE_BUILD_SHIPPING
+#if WITH_CHAOS_VISUAL_DEBUGGER
 	FChaosPathedMovementModeDebugData DebugData;
 #endif // !UE_BUILD_SHIPPING
 };
diff --git a/Engine/Plugins/Experimental/ChaosMover/Source/ChaosMover/Public/ChaosMover/PathedMovement/ChaosPathedMovementTypes.h b/Engine/Plugins/Experimental/ChaosMover/Source/ChaosMover/Public/ChaosMover/PathedMovement/ChaosPathedMovementTypes.h
--- a/Engine/Plugins/Experimental/ChaosMover/Source/ChaosMover/Public/ChaosMover/PathedMovement/ChaosPathedMovementTypes.h
+++ b/Engine/Plugins/Experimental/ChaosMover/Source/ChaosMover/Public/ChaosMover/PathedMovement/ChaosPathedMovementTypes.h
@@ -92,12 +92,19 @@ struct FChaosPathedMovementState : public FMoverDataStructBase
 {
 	GENERATED_BODY()
 
-	/** This is a value between 0 and 1 that can be directly used along the path to calculate
-	  * the current position on it, 0.0 being a the start of the progression along the path, 1.0 being at the end of that progression.
-	  * Note that this is not necessarily the start and end of that path as the start could be offseted relative to the path primitive (say, a spline)
+	/** A path progress is a value between 0 and 1 that can be directly used along the path to calculate
+	  * the position on it, 0.0 being a the start of the progression along the path, 1.0 being at the end of that progression.
+	  * Note that this is not necessarily the start and end of that path as the start could be offset relative to the path primitive (say, a spline)
+	  * LastChangePathProgress is the path progress we were at the last time an important change happened. We calculate the current
+	  * progress based on that LastChangePathProgress, the number of frames elapsed since that change, and the progression rate,
+	  * assuming all properties involved in updating that progress have not changed
 	  */
 	UPROPERTY(VisibleAnywhere, Category = ChaosMover)
-	float CurrentPathProgress = 0.0f;
+	float LastChangePathProgress = 0.0f;
+
+	// The frame at which the last change in movement properties has occurred. @see LastChangePathProgress
+	UPROPERTY(VisibleAnywhere, Category = ChaosMover)
+	int32 LastChangeFrame = INDEX_NONE;
 
 	/** The frame at which we last started playing, or INDEX_NONE if we were never playing */
 	UPROPERTY(VisibleAnywhere, Category = ChaosMover)
@@ -107,7 +114,7 @@ struct FChaosPathedMovementState : public FMoverDataStructBase
 	UPROPERTY(VisibleAnywhere, Category = ChaosMover)
 	bool bIsPathProgressionIncreasing = true;
 	
-	/** Whether the previous playback finished by reaching its end point.
+	/** Whether the previous playback finished by reaching the end of its overall progression (could be 0 or 1 depending on movement properties)
 	*/
 	UPROPERTY(VisibleAnywhere, Category = ChaosMover)
 	bool bHasFinished = false;
@@ -158,10 +165,28 @@ struct FChaosPathedMovementModeDebugData : public FMoverDataStructBase
 	FChaosPathedMovementState PostSimMoveState;
 
 	UPROPERTY(VisibleAnywhere, Category = "Mover Info")
-	FVector FinalWorldTranslation;
+	FVector PreviousRelativeTranslation;
+
+	UPROPERTY(VisibleAnywhere, Category = "Mover Info")
+	FRotator PreviousRelativeRotation;
+
+	UPROPERTY(VisibleAnywhere, Category = "Mover Info")
+	FVector PreviousWorldTranslation;
+
+	UPROPERTY(VisibleAnywhere, Category = "Mover Info")
+	FRotator PreviousWorldRotation;
+
+	UPROPERTY(VisibleAnywhere, Category = "Mover Info")
+	FVector NewRelativeTranslation;
+
+	UPROPERTY(VisibleAnywhere, Category = "Mover Info")
+	FRotator NewRelativeRotation;
+
+	UPROPERTY(VisibleAnywhere, Category = "Mover Info")
+	FVector NewWorldTranslation;
 
 	UPROPERTY(VisibleAnywhere, Category = "Mover Info")
-	FRotator FinalWorldRotation;
+	FRotator NewWorldRotation;
 
 	UPROPERTY(VisibleAnywhere, Category = "Mover Info")
 	bool bPlaybackStartedThisFrame = false;
@@ -322,6 +347,8 @@ public:
 // 0 being at the beginning of the progression, 1 being at the end of that progression. A given progression should
 // always result in the same resulting transform, but the spatial trajectory resulting from the progression
 // varying from 0 to 1 may loop or revisit the same curves multiple times
+// Implementing this interface implies the usage of a path basis (centralized in the simulation, not in modes)
+// that the path is relative to
 UINTERFACE(MinimalAPI)
 class UChaosPathedMovementModeInterface : public UInterface
 {
