Chaos Mover - Add support for crouching. This change adds:
- A stance modifier for updating the player capsule on the physics htread
- A crouching transition for triggering the stance modifier based on a crouch input in the input data
- The ability to override the max speed and acceleration on a movement mode
- A new physics thread overlap query used by the uncrouching code to test if the character can uncrouch

diff --git a/Engine/Plugins/Experimental/ChaosMover/Source/ChaosMover/Private/Backends/ChaosMoverAsyncCallback.cpp b/Engine/Plugins/Experimental/ChaosMover/Source/ChaosMover/Private/Backends/ChaosMoverAsyncCallback.cpp
--- a/Engine/Plugins/Experimental/ChaosMover/Source/ChaosMover/Private/Backends/ChaosMoverAsyncCallback.cpp
+++ b/Engine/Plugins/Experimental/ChaosMover/Source/ChaosMover/Private/Backends/ChaosMoverAsyncCallback.cpp
@@ -20,10 +20,13 @@ namespace UE::ChaosMover
 			return;
 		}
 
-		auto LambdaParallelUpdate = [&AsyncInput, PhysicsStep](int32 Idx) {
+		FMoverTimeStep TimeStep;
+		GeCurrentMoverTimeStep(AsyncInput, TimeStep);
+
+		auto LambdaParallelUpdate = [&TimeStep, &AsyncInput, PhysicsStep](int32 Idx) {
 			if (TStrongObjectPtr<UChaosMoverBackendComponent> Backend = AsyncInput->Backends[Idx].Pin())
 			{
-				Backend->GetSimulation()->ProcessInputs(PhysicsStep, AsyncInput->InputData[Idx]);
+				Backend->GetSimulation()->ProcessInputs(PhysicsStep, TimeStep, AsyncInput->InputData[Idx]);
 			}
 			};
 
diff --git a/Engine/Plugins/Experimental/ChaosMover/Source/ChaosMover/Private/ChaosMoverDeveloperSettings.cpp b/Engine/Plugins/Experimental/ChaosMover/Source/ChaosMover/Private/ChaosMoverDeveloperSettings.cpp
--- a/Engine/Plugins/Experimental/ChaosMover/Source/ChaosMover/Private/ChaosMoverDeveloperSettings.cpp
+++ b/Engine/Plugins/Experimental/ChaosMover/Source/ChaosMover/Private/ChaosMoverDeveloperSettings.cpp
@@ -20,6 +20,10 @@ namespace UE::ChaosMover
 		FAutoConsoleVariableRef CVarChaosMoverDrawGroundQueries(TEXT("ChaosMover.DebugDraw.GroundQueries"),
 			bDrawGroundQueries, TEXT("Draw ground queries."));
 
+		bool bDrawOverlapQueries = false;
+		FAutoConsoleVariableRef CVarChaosMoverDrawOverlapQueries(TEXT("ChaosMover.DebugDraw.OverlapQueries"),
+			bDrawOverlapQueries, TEXT("Draw overlap queries."));
+
 		bool bSkipGenerateMoveIfOverridden = true;
 		FAutoConsoleVariableRef CVarChaosMoverSkipGenerateMoveIfOverridden(TEXT("ChaosMover.Perf.SkipGenerateMoveIfOverridden"),
 			bSkipGenerateMoveIfOverridden, TEXT("If true and we have a layered move fully overriding movement, then we will skip calling OnGenerateMove on the active movement mode for better performance\n")
diff --git a/Engine/Plugins/Experimental/ChaosMover/Source/ChaosMover/Private/ChaosMoverSimulation.cpp b/Engine/Plugins/Experimental/ChaosMover/Source/ChaosMover/Private/ChaosMoverSimulation.cpp
--- a/Engine/Plugins/Experimental/ChaosMover/Source/ChaosMover/Private/ChaosMoverSimulation.cpp
+++ b/Engine/Plugins/Experimental/ChaosMover/Source/ChaosMover/Private/ChaosMoverSimulation.cpp
@@ -12,6 +12,7 @@
 #include "ChaosMover/ChaosMovementModeTransition.h"
 #include "ChaosMover/ChaosMoverLog.h"
 #include "ChaosMover/ChaosMoverSimulationTypes.h"
+#include "ChaosMover/Character/ChaosCharacterInputs.h"
 #include "ChaosMover/Utilities/ChaosGroundMovementUtils.h"
 #include "ChaosMover/Utilities/ChaosMoverQueryUtils.h"
 #include "ChaosMover/PathedMovement/ChaosPathedMovementTypes.h"
@@ -69,6 +70,11 @@ const UBaseMovementMode* UChaosMoverSimulation::FindMovementModeByName(const FNa
 	return StateMachine.FindMovementMode(Name).Get();
 }
 
+UBaseMovementMode* UChaosMoverSimulation::FindMovementModeByName_Mutable(const FName& Name)
+{
+	return StateMachine.FindMovementMode_Mutable(Name).Get();
+}
+
 void UChaosMoverSimulation::ApplyNetInputData(const FMoverInputCmdContext& InNetInputCmd)
 {
 	NetInputCmd = InNetInputCmd;
@@ -82,12 +88,13 @@ void UChaosMoverSimulation::BuildNetInputData(FMoverInputCmdContext& OutNetInput
 
 void UChaosMoverSimulation::ApplyNetStateData(const FMoverSyncState& InNetSyncState)
 {
-	CurrentSyncState = InNetSyncState;
+	NetSyncState = InNetSyncState;
+	bSyncStateOverridden = true;
 }
 
 void UChaosMoverSimulation::BuildNetStateData(FMoverSyncState& OutNetSyncState) const
 {
-	OutNetSyncState = CurrentSyncState;
+	OutNetSyncState = NetSyncState;
 }
 
 void UChaosMoverSimulation::Init(const FInitParams& InitParams)
@@ -166,7 +173,7 @@ void UChaosMoverSimulation::OnDeinit()
 
 }
 
-void UChaosMoverSimulation::ProcessInputs(int32 PhysicsStep, const UE::ChaosMover::FSimulationInputData& InputData)
+void UChaosMoverSimulation::ProcessInputs(int32 PhysicsStep, const FMoverTimeStep& TimeStep, const UE::ChaosMover::FSimulationInputData& InputData)
 {
 	Chaos::EnsureIsInPhysicsThreadContext();
 
@@ -182,6 +189,23 @@ void UChaosMoverSimulation::ProcessInputs(int32 PhysicsStep, const UE::ChaosMove
 	{
 		NetInputCmd = InputData.InputCmd;
 	}
+
+	// Sync state is overwritten by network physics for sim proxies and for resimulation
+	// In either case we rollback the simulation to invalidate the blackboard and reset
+	// the state machine based on the new state
+	if (bSyncStateOverridden)
+	{
+		const FMoverSyncState PrevSyncState = CurrentSyncState;
+		CurrentSyncState = NetSyncState;
+		OnSimulationRollback(TimeStep, PrevSyncState);
+		bSyncStateOverridden = false;
+	}
+	else
+	{
+		NetSyncState = CurrentSyncState;
+	}
+
+	
 }
 
 void UChaosMoverSimulation::SimulationTick(const FMoverTimeStep& TimeStep, const UE::ChaosMover::FSimulationInputData& InputData, UE::ChaosMover::FSimulationOutputData& OutputData)
@@ -271,6 +295,34 @@ void UChaosMoverSimulation::PreSimulationTickCharacter(const IChaosCharacterMove
 		CharacterDefaultInputs.SuggestedMovementMode = NAME_None;
 	}
 
+	if (const FChaosMovementSettingsOverrides* MovementSettingsOverride = InputData.InputCmd.InputCollection.FindDataByType<FChaosMovementSettingsOverrides>())
+	{
+		// If the override contains a mode name find that mode to override. Otherwise take the latest mode,
+		// or the pending mode if a new one has been set
+		TWeakObjectPtr<UBaseMovementMode> ModePtr = !MovementSettingsOverride->ModeName.IsNone() ? StateMachine.FindMovementMode_Mutable(MovementSettingsOverride->ModeName) :
+			(CharacterDefaultInputs.SuggestedMovementMode.IsNone() ? StateMachine.GetCurrentMode() : StateMachine.FindMovementMode_Mutable(CharacterDefaultInputs.SuggestedMovementMode));
+
+		if (IChaosCharacterMovementModeInterface* ModeInterface = Cast<IChaosCharacterMovementModeInterface>(ModePtr.Get()))
+		{
+			ModeInterface->OverrideMaxSpeed(MovementSettingsOverride->MaxSpeedOverride);
+			ModeInterface->OverrideAcceleration(MovementSettingsOverride->AccelerationOverride);
+		}
+	}
+
+	if (const FChaosMovementSettingsOverridesRemover* MovementSettingsOverrideRemover = InputData.InputCmd.InputCollection.FindDataByType<FChaosMovementSettingsOverridesRemover>())
+	{
+		// If the override contains a mode name find that mode to override. Otherwise take the latest mode,
+		// or the pending mode if a new one has been set
+		TWeakObjectPtr<UBaseMovementMode> ModePtr = !MovementSettingsOverrideRemover->ModeName.IsNone() ? StateMachine.FindMovementMode_Mutable(MovementSettingsOverrideRemover->ModeName) :
+			(CharacterDefaultInputs.SuggestedMovementMode.IsNone() ? StateMachine.GetCurrentMode() : StateMachine.FindMovementMode_Mutable(CharacterDefaultInputs.SuggestedMovementMode));
+
+		if (IChaosCharacterMovementModeInterface* ModeInterface = Cast<IChaosCharacterMovementModeInterface>(ModePtr.Get()))
+		{
+			ModeInterface->ClearMaxSpeedOverride();
+			ModeInterface->ClearAccelerationOverride();
+		}
+	}
+
 	// The state machine works in local space relative to the ground, so if we're on
 	// the ground make the state relative to the ground
 	// TODO - Can we use the local space in the sync state rather than subtracting and
@@ -302,18 +354,16 @@ void UChaosMoverSimulation::PreSimulationTickCharacter(const IChaosCharacterMove
 	GroundSimState.LocalVelocity = LocalGroundVelocity;
 }
 
-void UChaosMoverSimulation::OnPreSimulationTick(const FMoverTimeStep& TimeStep, const UE::ChaosMover::FSimulationInputData& InputData)
+void UChaosMoverSimulation::OnSimulationRollback(const FMoverTimeStep& NewTimeStep, const FMoverSyncState& PrevSyncState)
 {
-	if (TimeStep.bIsResimulating)
-	{
-		check(Solver);
-		if (Solver->GetEvolution()->IsResetting())
-		{
-			// Rollback blackboard on the first frame of resimulation
-			Blackboard->Invalidate(EInvalidationReason::Rollback);
-		}
-	}
+	// Rollback blackboard on the first frame of resimulation
+	Blackboard->Invalidate(EInvalidationReason::Rollback);
+
+	StateMachine.OnSimulationRollback(NewTimeStep, PrevSyncState, CurrentSyncState);
+}
 
+void UChaosMoverSimulation::OnPreSimulationTick(const FMoverTimeStep& TimeStep, const UE::ChaosMover::FSimulationInputData& InputData)
+{
 	// Update the sync state from the current physics state
 	FMoverDefaultSyncState& PreSimDefaultSyncState = CurrentSyncState.SyncStateCollection.FindOrAddMutableDataByType<FMoverDefaultSyncState>();
 	PreSimDefaultSyncState.SetMovementBase(nullptr);
@@ -429,6 +479,7 @@ void UChaosMoverSimulation::PostSimulationTickCharacterConstraint(const IChaosCh
 	if (const FChaosMoverSimulationDefaultInputs* SimInputs = LocalSimInput.FindDataByType<FChaosMoverSimulationDefaultInputs>())
 	{
 		Settings.VerticalAxis = SimInputs->UpDir;
+		Settings.TargetHeight = CharacterConstraintMode.GetTargetHeight();
 	}
 
 	check(Blackboard.Get());
@@ -903,3 +954,80 @@ void UChaosMoverSimulation::QueueInstantMovementEffect(TSharedPtr<FInstantMoveme
 {
 	StateMachine.QueueInstantMovementEffect(InstantMovementEffect);
 }
+
+DEFINE_FUNCTION(UChaosMoverSimulation::execK2_QueueMovementModifier)
+{
+	Stack.StepCompiledIn<FStructProperty>(nullptr);
+	void* MovePtr = Stack.MostRecentPropertyAddress;
+	FStructProperty* StructProp = CastField<FStructProperty>(Stack.MostRecentProperty);
+
+	P_FINISH;
+
+	P_NATIVE_BEGIN;
+
+	const bool bHasValidStructProp = StructProp && StructProp->Struct && StructProp->Struct->IsChildOf(FMovementModifierBase::StaticStruct());
+
+	if (ensureMsgf((bHasValidStructProp && MovePtr), TEXT("An invalid type (%s) was sent to a QueueMovementModifier node. A struct derived from FMovementModifierBase is required. No modifier will be queued."),
+		StructProp ? *GetNameSafe(StructProp->Struct) : *Stack.MostRecentProperty->GetClass()->GetName()))
+	{
+		// Could we steal this instead of cloning? (move semantics)
+		FMovementModifierBase* MoveAsBasePtr = reinterpret_cast<FMovementModifierBase*>(MovePtr);
+		FMovementModifierBase* ClonedMove = MoveAsBasePtr->Clone();
+
+		FMovementModifierHandle ModifierID = P_THIS->QueueMovementModifier(TSharedPtr<FMovementModifierBase>(ClonedMove));
+		*static_cast<FMovementModifierHandle*>(RESULT_PARAM) = ModifierID;
+	}
+
+	P_NATIVE_END;
+}
+
+
+FMovementModifierHandle UChaosMoverSimulation::QueueMovementModifier(TSharedPtr<FMovementModifierBase> Modifier)
+{
+	return StateMachine.QueueMovementModifier(Modifier);
+}
+
+void UChaosMoverSimulation::CancelModifierFromHandle(FMovementModifierHandle ModifierHandle)
+{
+	StateMachine.CancelModifierFromHandle(ModifierHandle);
+}
+
+bool UChaosMoverSimulation::HasGameplayTag(FGameplayTag TagToFind, bool bExactMatch) const
+{
+	// Check active Movement Mode
+	if (const UBaseMovementMode* ActiveMovementMode = FindMovementModeByName(CurrentSyncState.MovementMode))
+	{
+		if (ActiveMovementMode->HasGameplayTag(TagToFind, bExactMatch))
+		{
+			return true;
+		}
+	}
+
+	// Search Movement Modifiers
+	for (auto ModifierFromSyncStateIt = CurrentSyncState.MovementModifiers.GetActiveModifiersIterator(); ModifierFromSyncStateIt; ++ModifierFromSyncStateIt)
+	{
+		if (const TSharedPtr<FMovementModifierBase> ModifierFromSyncState = *ModifierFromSyncStateIt)
+		{
+			if (ModifierFromSyncState.IsValid() && ModifierFromSyncState->HasGameplayTag(TagToFind, bExactMatch))
+			{
+				return true;
+			}
+		}
+	}
+
+	// Search Layered Moves
+	for (const TSharedPtr<FLayeredMoveBase>& LayeredMove : CurrentSyncState.LayeredMoves.GetActiveMoves())
+	{
+		if (LayeredMove->HasGameplayTag(TagToFind, bExactMatch))
+		{
+			return true;
+		}
+	}
+
+	return false;
+}
+
+const FMovementModifierBase* UChaosMoverSimulation::FindMovementModifierByType(const UScriptStruct* DataStructType) const
+{
+	return StateMachine.FindQueuedModifierByType(DataStructType);
+}
diff --git a/Engine/Plugins/Experimental/ChaosMover/Source/ChaosMover/Private/ChaosMoverStateMachine.cpp b/Engine/Plugins/Experimental/ChaosMover/Source/ChaosMover/Private/ChaosMoverStateMachine.cpp
--- a/Engine/Plugins/Experimental/ChaosMover/Source/ChaosMover/Private/ChaosMoverStateMachine.cpp
+++ b/Engine/Plugins/Experimental/ChaosMover/Source/ChaosMover/Private/ChaosMoverStateMachine.cpp
@@ -166,25 +166,114 @@ namespace UE::ChaosMover
 		}
 	}
 
-	void FMoverStateMachine::OnSimulationTick(const FMoverTimeStep& TimeStep, const FMoverTickStartData& StartState, UMoverBlackboard* SimBlackboard, UMovementMixer* MovementMixer, FMoverTickEndData& OutputState)
+	FMovementModifierHandle FMoverStateMachine::QueueMovementModifier(TSharedPtr<FMovementModifierBase> Modifier)
 	{
-		if (!ensure(MovementMixer))
+		if (ensure(Modifier.IsValid()))
 		{
-			return;
+			QueuedMovementModifiers.Add(Modifier);
+			Modifier->GenerateHandle();
+
+			return Modifier->GetHandle();
 		}
 
-		InternalSimTimeMs = TimeStep.BaseSimTimeMs;
+		return 0;
+	}
+
+	void FMoverStateMachine::CancelModifierFromHandle(FMovementModifierHandle ModifierHandle)
+	{
+		QueuedMovementModifiers.RemoveAll([ModifierHandle, this] (const TSharedPtr<FMovementModifierBase>& Modifier)
+			{
+				if (Modifier.IsValid())
+				{
+					if (Modifier->GetHandle() == ModifierHandle)
+					{
+						return true;
+					}
+				}
+				else
+				{
+					return true;
+				}
 
+				return false;
+			});
+
+		ModifiersToCancel.Add(ModifierHandle);
+	}
+
+	const FMovementModifierBase* FMoverStateMachine::FindQueuedModifier(FMovementModifierHandle ModifierHandle) const
+	{
+		for (const TSharedPtr<FMovementModifierBase>& QueuedModifier : QueuedMovementModifiers)
+		{
+			if (QueuedModifier->GetHandle() == ModifierHandle)
+			{
+				return QueuedModifier.Get();
+			}
+		}
+
+		return nullptr;
+	}
+
+	const FMovementModifierBase* FMoverStateMachine::FindQueuedModifierByType(const UScriptStruct* ModifierType) const
+	{
+		for (const TSharedPtr<FMovementModifierBase>& QueuedModifier : QueuedMovementModifiers)
+		{
+			if (QueuedModifier->GetScriptStruct() == ModifierType)
+			{
+				return QueuedModifier.Get();
+			}
+		}
+
+		return nullptr;
+	}
+
+	void FMoverStateMachine::FlushQueuedModifiersToGroup(FMovementModifierGroup& ModifierGroup)
+	{
+		if (!QueuedMovementModifiers.IsEmpty())
+		{
+			for (TSharedPtr<FMovementModifierBase>& QueuedModifier : QueuedMovementModifiers)
+			{
+				ModifierGroup.QueueMovementModifier(QueuedModifier);
+			}
+
+			QueuedMovementModifiers.Empty();
+		}
+	}
+
+	void FMoverStateMachine::FlushModifierCancellationsToGroup(FMovementModifierGroup& ActiveModifierGroup)
+	{
+		for (FMovementModifierHandle HandleToCancel : ModifiersToCancel)
+		{
+			ActiveModifierGroup.CancelModifierFromHandle(HandleToCancel);
+		}
+
+		ModifiersToCancel.Empty();
+	}
+
+	bool FMoverStateMachine::HasAnyInstantEffectsQueued() const
+	{
+		return !QueuedInstantEffects.IsEmpty();
+	}
+
+	void FMoverStateMachine::OnSimulationTick(const FMoverTimeStep& TimeStep, const FMoverTickStartData& StartState, UMoverBlackboard* SimBlackboard, UMovementMixer* MovementMixer, FMoverTickEndData& OutputState)
+	{
 		//GEngine->AddOnScreenDebugMessage(-1, -0.1f, FColor::White, FString::Printf(TEXT("Mode Tick: %s  Queued: %s"), *CurrentModeName.ToString(), *NextModeName.ToString()));
 
 		FMoverTimeStep SubTimeStep = TimeStep;
-		FMoverTickStartData SubstepStartData = StartState;
-		FMoverDefaultSyncState* SubstepStartSyncState = SubstepStartData.SyncState.SyncStateCollection.FindMutableDataByType<FMoverDefaultSyncState>();
+		WorkingSubstepStartData = StartState;
+		bool bIsWorkingStartStateReady = true;	// this flag is used to avoid unneeded data copying after substeps
+
+		if (!ensure(MovementMixer))
+		{
+			return;
+		}
+
+		InternalSimTimeMs = TimeStep.BaseSimTimeMs;
 
 		TStrongObjectPtr<UImmediateMovementModeTransition> QueuedModeTransition = QueuedModeTransitionWeakPtr.Pin();
 		if (QueuedModeTransition && !QueuedModeTransition->IsSet())
 		{
-			QueueNextMode(SubstepStartData.SyncState.MovementMode);
+			QueueNextMode(WorkingSubstepStartData.SyncState.MovementMode);
 		}
 
 		AdvanceToNextMode();
@@ -198,7 +287,15 @@ namespace UE::ChaosMover
 		{
 			InternalSimTimeMs = SubTimeStep.BaseSimTimeMs;
 
-			SubstepStartSyncState = SubstepStartData.SyncState.SyncStateCollection.FindMutableDataByType<FMoverDefaultSyncState>();
+			if (!bIsWorkingStartStateReady)
+			{
+				WorkingSubstepStartData.SyncState = OutputState.SyncState;
+				WorkingSubstepStartData.AuxState = OutputState.AuxState;
+				bIsWorkingStartStateReady = true;
+			}
+
+			WorkingSubstepStartData.SyncState.MovementMode = CurrentModeName;
+
 			FMoverDefaultSyncState* OutputSyncState = &OutputState.SyncState.SyncStateCollection.FindOrAddMutableDataByType<FMoverDefaultSyncState>();
 			OutputState.SyncState.MovementMode = CurrentModeName;
 
@@ -207,19 +304,22 @@ namespace UE::ChaosMover
 			SubTimeStep.StepMs = TimeStep.StepMs - TotalUsedMs;		// TODO: convert this to an overridable function that can support MaxStepTime, MaxIterations, etc.
 
 			// Transfer any queued moves into the starting state. They'll be started during the move generation.
-			FlushQueuedMovesToGroup(SubstepStartData.SyncState.LayeredMoves);
-			OutputState.SyncState.LayeredMoves = SubstepStartData.SyncState.LayeredMoves;
+			FlushQueuedMovesToGroup(WorkingSubstepStartData.SyncState.LayeredMoves);
+			OutputState.SyncState.LayeredMoves = WorkingSubstepStartData.SyncState.LayeredMoves;
+
+			FlushQueuedModifiersToGroup(WorkingSubstepStartData.SyncState.MovementModifiers);
+			OutputState.SyncState.MovementModifiers = WorkingSubstepStartData.SyncState.MovementModifiers;
 
 			bool bModeSetFromInstantEffect = false;
 			// Apply any instant effects that were queued up between ticks
-			if (ApplyInstantEffects(SubstepStartData, SubTimeStep, OutputState.SyncState))
+			if (ApplyInstantEffects(WorkingSubstepStartData, SubTimeStep, OutputState.SyncState))
 			{
 				// Copying over our sync state collection to SubstepStartData so it is effectively the input sync state later for the movement mode. Doing this makes sure state modification from Instant Effects isn't overridden later by the movement mode
 				for (auto SyncDataIt = OutputState.SyncState.SyncStateCollection.GetCollectionDataIterator(); SyncDataIt; ++SyncDataIt)
 				{
 					if (SyncDataIt->Get())
 					{
-						SubstepStartData.SyncState.SyncStateCollection.AddOrOverwriteData(TSharedPtr<FMoverDataStructBase>(SyncDataIt->Get()->Clone()));
+						WorkingSubstepStartData.SyncState.SyncStateCollection.AddOrOverwriteData(TSharedPtr<FMoverDataStructBase>(SyncDataIt->Get()->Clone()));
 					}
 				}
 
@@ -227,13 +327,20 @@ namespace UE::ChaosMover
 				{
 					bModeSetFromInstantEffect = true;
 					SetModeImmediately(OutputState.SyncState.MovementMode);
-					SubstepStartData.SyncState.MovementMode = CurrentModeName;
+					WorkingSubstepStartData.SyncState.MovementMode = CurrentModeName;
 				}
 			}
 
+			FMovementModifierParams_Async ModifierParams(Simulation.Get(), &WorkingSubstepStartData.SyncState, &SubTimeStep);
+
 			FMovementModifierGroup& CurrentModifiers = OutputState.SyncState.MovementModifiers;
-			TArray<TSharedPtr<FMovementModifierBase>> ActiveModifiers;
-			ensureMsgf(!CurrentModifiers.HasAnyMoves(), TEXT("Movement Modifiers are not supported in Mover async mode"));
+			FlushModifierCancellationsToGroup(CurrentModifiers);
+			TArray<TSharedPtr<FMovementModifierBase>> ActiveModifiers = CurrentModifiers.GenerateActiveModifiers_Async(ModifierParams);
+
+			for (TSharedPtr<FMovementModifierBase> Modifier : ActiveModifiers)
+			{
+				Modifier->OnPreMovement_Async(ModifierParams);
+			}
 
 			FLayeredMoveGroup& CurrentLayeredMoves = OutputState.SyncState.LayeredMoves;
 
@@ -251,8 +358,9 @@ namespace UE::ChaosMover
 			for (TSharedPtr<FLayeredMoveBase>& ActiveMove : ActiveMoves)
 			{
 				FProposedMove MoveStep;
-				bool GenerateMoveResult = ActiveMove->GenerateMove_Async(SubstepStartData, SubTimeStep, SimBlackboard, MoveStep);
-				if (GenerateMoveResult)
+				MoveStep.MixMode = ActiveMove->MixMode;	// Initialize using the move's mixmode, but allow it to be changed in GenerateMove
+
+				if (ActiveMove->GenerateMove_Async(WorkingSubstepStartData, SubTimeStep, SimBlackboard, MoveStep))
 				{
 					// If this active move is already past it's first tick we don't need to set the preferred mode again
 					if (ActiveMove->StartSimTimeMs < SubTimeStep.BaseSimTimeMs)
@@ -282,7 +390,7 @@ namespace UE::ChaosMover
 					!(bHasLayeredMoveContributions && CombinedLayeredMove.MixMode == EMoveMixMode::OverrideAll))
 				{
 					QUICK_SCOPE_CYCLE_COUNTER(STAT_GenerateMoveFromMode);
-					CurrentMode->GenerateMove(SubstepStartData, SubTimeStep, OUT CombinedMove);
+					CurrentMode->GenerateMove(WorkingSubstepStartData, SubTimeStep, OUT CombinedMove);
 
 					bHasModeMoveContribution = true;
 				}
@@ -318,12 +426,11 @@ namespace UE::ChaosMover
 
 				// Execute the combined proposed move
 				{
-					FSimulationTickParams SimTickParams;
-					// SimTickParams.MovingComps is left empty in the async case, so we don't access resources used by the concurrent gameplay thread
-					SimTickParams.SimBlackboard = SimBlackboard;
-					SimTickParams.StartState = SubstepStartData;
-					SimTickParams.TimeStep = SubTimeStep;
-					SimTickParams.ProposedMove = CombinedMove;
+					// WorkingSimTickParams.MovingComps is left empty in the async case, so we don't access resources used by the concurrent gameplay thread
+					WorkingSimTickParams.StartState = WorkingSubstepStartData;
+					WorkingSimTickParams.SimBlackboard = SimBlackboard;
+					WorkingSimTickParams.TimeStep = SubTimeStep;
+					WorkingSimTickParams.ProposedMove = CombinedMove;
 
 					// Check for any transitions, first those registered with the current movement mode, then global ones that could occur from any mode
 					FTransitionEvalResult EvalResult = FTransitionEvalResult::NoTransition;
@@ -333,7 +440,7 @@ namespace UE::ChaosMover
 					{
 						if (IsValid(Transition) && ((SubStepCount == 0) || !Transition->bFirstSubStepOnly))
 						{
-							EvalResult = Transition->Evaluate(SimTickParams);
+							EvalResult = Transition->Evaluate(WorkingSimTickParams);
 
 							if (!EvalResult.NextMode.IsNone())
 							{
@@ -353,7 +460,7 @@ namespace UE::ChaosMover
 							TStrongObjectPtr<UBaseMovementModeTransition> Transition = TransitionWeakPtr.Pin();
 							if (Transition)
 							{
-								EvalResult = Transition->Evaluate(SimTickParams);
+								EvalResult = Transition->Evaluate(WorkingSimTickParams);
 
 								if (!EvalResult.NextMode.IsNone())
 								{
@@ -370,12 +477,12 @@ namespace UE::ChaosMover
 					if (TransitionToTrigger && !EvalResult.NextMode.IsNone())
 					{
 						OutputState.MovementEndState.NextModeName = EvalResult.NextMode;
-						OutputState.MovementEndState.RemainingMs = SimTickParams.TimeStep.StepMs; 	// Pass all remaining time to next mode
-						TransitionToTrigger->Trigger(SimTickParams);
+						OutputState.MovementEndState.RemainingMs = WorkingSimTickParams.TimeStep.StepMs; 	// Pass all remaining time to next mode
+						TransitionToTrigger->Trigger(WorkingSimTickParams);
 					}
 					else
 					{
-						CurrentMode->SimulationTick(SimTickParams, OutputState);
+						CurrentMode->SimulationTick(WorkingSimTickParams, OutputState);
 					}
 
 					OutputState.MovementEndState.RemainingMs = FMath::Clamp(OutputState.MovementEndState.RemainingMs, 0.0f, SubTimeStep.StepMs);
@@ -409,7 +516,10 @@ namespace UE::ChaosMover
 			AdvanceToNextMode();
 			OutputState.SyncState.MovementMode = CurrentModeName;
 
-			ensureMsgf(ActiveModifiers.IsEmpty(), TEXT("Movement Modifiers are not supported in Async mode"));
+			for (TSharedPtr<FMovementModifierBase> Modifier : ActiveModifiers)
+			{
+				Modifier->OnPostMovement_Async(ModifierParams);
+			}
 
 			const float RemainingMs = FMath::Clamp(OutputState.MovementEndState.RemainingMs, 0.0f, SubTimeStep.StepMs);
 			const float SubstepUsedMs = (SubTimeStep.StepMs - RemainingMs);
@@ -417,33 +527,91 @@ namespace UE::ChaosMover
 			TotalUsedMs += SubstepUsedMs;
 			SubTimeStep.StepMs = RemainingMs;
 
-			SubstepStartData.SyncState = OutputState.SyncState;
-			SubstepStartData.AuxState = OutputState.AuxState;
-
+			bIsWorkingStartStateReady = false;
 			++SubStepCount;
 		}
 
 		InternalSimTimeMs = TimeStep.BaseSimTimeMs + TotalUsedMs;
 
-		// Apply any instant effects that were queued up during this tick and didn't get handled in a substep
-		if (ApplyInstantEffects(SubstepStartData, SubTimeStep, OutputState.SyncState))
+		if (HasAnyInstantEffectsQueued())
 		{
-			if (CurrentModeName != OutputState.SyncState.MovementMode)
+			if (!bIsWorkingStartStateReady)
 			{
-				SetModeImmediately(OutputState.SyncState.MovementMode);
+				WorkingSubstepStartData.SyncState = OutputState.SyncState;
+				WorkingSubstepStartData.AuxState = OutputState.AuxState;
+				bIsWorkingStartStateReady = true;
 			}
-		}
-	}
 
-	void FMoverStateMachine::OnSimulationPreRollback(const FMoverSyncState* InvalidSyncState, const FMoverSyncState* SyncState, const FMoverAuxStateContext* InvalidAuxState, const FMoverAuxStateContext* AuxState)
-	{
+			// Apply any instant effects that were queued up during this tick and didn't get handled in a substep
+			if (ApplyInstantEffects(WorkingSubstepStartData, SubTimeStep, OutputState.SyncState))
+			{
+				if (CurrentModeName != OutputState.SyncState.MovementMode)
+				{
+					SetModeImmediately(OutputState.SyncState.MovementMode);
+				}
+			}
+		}
 	}
 
-	void FMoverStateMachine::OnSimulationRollback(const FMoverSyncState* SyncState, const FMoverAuxStateContext* AuxState)
+	void FMoverStateMachine::OnSimulationRollback(const FMoverTimeStep& NewTimeStep, const FMoverSyncState& InvalidSyncState, const FMoverSyncState& NewSyncState)
 	{
 		ClearQueuedMode();
+
+		if (CurrentModeName != NewSyncState.MovementMode)
+		{
+			SetModeImmediately(NewSyncState.MovementMode);
+		}
+
 		QueuedLayeredMoves.Empty();
 		QueuedInstantEffects.Empty();
+		QueuedMovementModifiers.Empty();
+
+		FMovementModifierParams_Async ModifierParams(Simulation.Get(), &NewSyncState, &NewTimeStep);
+
+		// Check if we have a new modifier in the rolled back sync state
+		for (auto ModifierFromRollbackIt = NewSyncState.MovementModifiers.GetActiveModifiersIterator(); ModifierFromRollbackIt; ++ModifierFromRollbackIt)
+		{
+			bool bContainsModifier = false;
+			for (auto ModifierFromCacheIt = InvalidSyncState.MovementModifiers.GetActiveModifiersIterator(); ModifierFromCacheIt; ++ModifierFromCacheIt)
+			{
+				if (ModifierFromRollbackIt->Get()->Matches(ModifierFromCacheIt->Get()))
+				{
+					bContainsModifier = true;
+
+					// Rolled back version of the modifier will be missing the handle; we fix that here
+					ModifierFromRollbackIt->Get()->OverwriteHandleIfInvalid(ModifierFromCacheIt->Get()->GetHandle());
+					break;
+				}
+			}
+
+			// If modifier is not already present start the new one
+			if (!bContainsModifier)
+			{
+				UE_LOG(LogChaosMover, Verbose, TEXT("Modifier(%s) was started after a rollback."), *ModifierFromRollbackIt->Get()->ToSimpleString());
+				ModifierFromRollbackIt->Get()->OnStart_Async(ModifierParams);
+			}
+		}
+
+		// Check if the previous state has an active modifier not in the rolled back state
+		for (auto ModifierFromCacheIt = InvalidSyncState.MovementModifiers.GetActiveModifiersIterator(); ModifierFromCacheIt; ++ModifierFromCacheIt)
+		{
+			bool bContainsModifier = false;
+			for (auto ModifierFromRollbackIt = NewSyncState.MovementModifiers.GetActiveModifiersIterator(); ModifierFromRollbackIt; ++ModifierFromRollbackIt)
+			{
+				if (ModifierFromRollbackIt->Get()->Matches(ModifierFromCacheIt->Get()))
+				{
+					bContainsModifier = true;
+					break;
+				}
+			}
+
+			// If the modifier is not in the rolled back state end it
+			if (!bContainsModifier)
+			{
+				UE_LOG(LogChaosMover, Log, TEXT("Modifier(%s) was ended after a rollback."), *ModifierFromCacheIt->Get()->ToSimpleString());
+				ModifierFromCacheIt->Get()->OnEnd_Async(ModifierParams);
+			}
+		}
 	}
 
 
@@ -467,6 +635,16 @@ namespace UE::ChaosMover
 		return nullptr;
 	}
 
+	TWeakObjectPtr<UBaseMovementMode> FMoverStateMachine::FindMovementMode_Mutable(FName ModeName)
+	{
+		if (ModeName != NAME_None && Modes.Contains(ModeName))
+		{
+			return Modes[ModeName];
+		}
+
+		return nullptr;
+	}
+
 	void FMoverStateMachine::QueueLayeredMove(TSharedPtr<FLayeredMoveBase> Move)
 	{
 		QueuedLayeredMoves.Add(Move);
diff --git a/Engine/Plugins/Experimental/ChaosMover/Source/ChaosMover/Private/ChaosMoverStateMachine.h b/Engine/Plugins/Experimental/ChaosMover/Source/ChaosMover/Private/ChaosMoverStateMachine.h
--- a/Engine/Plugins/Experimental/ChaosMover/Source/ChaosMover/Private/ChaosMoverStateMachine.h
+++ b/Engine/Plugins/Experimental/ChaosMover/Source/ChaosMover/Private/ChaosMoverStateMachine.h
@@ -13,11 +13,11 @@ class UChaosMoverSimulation;
 
 namespace UE::ChaosMover
 {
-	class CHAOSMOVER_API FMoverStateMachine
+	class FMoverStateMachine
 	{
 	public:
-		FMoverStateMachine();
-		virtual ~FMoverStateMachine();
+		CHAOSMOVER_API FMoverStateMachine();
+		CHAOSMOVER_API virtual ~FMoverStateMachine();
 
 		struct FInitParams
 		{
@@ -25,40 +25,44 @@ namespace UE::ChaosMover
 			TWeakObjectPtr<UNullMovementMode> NullMovementMode;
 			TWeakObjectPtr<UImmediateMovementModeTransition> ImmediateMovementModeTransition;
 		};
-		void Init(const FInitParams& Params);
+		CHAOSMOVER_API void Init(const FInitParams& Params);
 
-		void RegisterMovementMode(FName ModeName, TWeakObjectPtr<UBaseMovementMode> Mode, bool bIsDefaultMode = false);
+		CHAOSMOVER_API void RegisterMovementMode(FName ModeName, TWeakObjectPtr<UBaseMovementMode> Mode, bool bIsDefaultMode = false);
 
-		void UnregisterMovementMode(FName ModeName);
-		void ClearAllMovementModes();
+		CHAOSMOVER_API void UnregisterMovementMode(FName ModeName);
+		CHAOSMOVER_API void ClearAllMovementModes();
 
-		void RegisterGlobalTransition(TWeakObjectPtr<UBaseMovementModeTransition> Transition);
-		void UnregisterGlobalTransition(TWeakObjectPtr<UBaseMovementModeTransition> Transition);
-		void ClearAllGlobalTransitions();
+		CHAOSMOVER_API void RegisterGlobalTransition(TWeakObjectPtr<UBaseMovementModeTransition> Transition);
+		CHAOSMOVER_API void UnregisterGlobalTransition(TWeakObjectPtr<UBaseMovementModeTransition> Transition);
+		CHAOSMOVER_API void ClearAllGlobalTransitions();
 
-		FName GetDefaultModeName() const;
-		void SetDefaultMode(FName NewDefaultModeName);
+		CHAOSMOVER_API FName GetDefaultModeName() const;
+		CHAOSMOVER_API void SetDefaultMode(FName NewDefaultModeName);
 
-		void QueueNextMode(FName DesiredNextModeName, bool bShouldReenter = false);
-		void SetModeImmediately(FName DesiredModeName, bool bShouldReenter = false);
-		void ClearQueuedMode();
+		CHAOSMOVER_API void QueueNextMode(FName DesiredNextModeName, bool bShouldReenter = false);
+		CHAOSMOVER_API void SetModeImmediately(FName DesiredModeName, bool bShouldReenter = false);
+		CHAOSMOVER_API void ClearQueuedMode();
 
-		void OnSimulationTick(const FMoverTimeStep& TimeStep, const FMoverTickStartData& StartState, UMoverBlackboard* SimBlackboard, UMovementMixer* MovementMixer, FMoverTickEndData& OutputState);
-		void OnSimulationPreRollback(const FMoverSyncState* InvalidSyncState, const FMoverSyncState* SyncState, const FMoverAuxStateContext* InvalidAuxState, const FMoverAuxStateContext* AuxState);
-		void OnSimulationRollback(const FMoverSyncState* SyncState, const FMoverAuxStateContext* AuxState);
+		CHAOSMOVER_API void OnSimulationTick(const FMoverTimeStep& TimeStep, const FMoverTickStartData& StartState, UMoverBlackboard* SimBlackboard, UMovementMixer* MovementMixer, FMoverTickEndData& OutputState);
+		CHAOSMOVER_API void OnSimulationRollback(const FMoverTimeStep& NewTimeStep, const FMoverSyncState& InvalidSyncState, const FMoverSyncState& NewSyncState);
 
 		FName GetCurrentModeName() const
 		{
 			return CurrentModeName;
 		}
 
-		const TWeakObjectPtr<UBaseMovementMode> GetCurrentMode() const;
+		CHAOSMOVER_API const TWeakObjectPtr<UBaseMovementMode> GetCurrentMode() const;
+		CHAOSMOVER_API const TWeakObjectPtr<UBaseMovementMode> FindMovementMode(FName ModeName) const;
+		CHAOSMOVER_API TWeakObjectPtr<UBaseMovementMode> FindMovementMode_Mutable(FName ModeName);
 
-		const TWeakObjectPtr<UBaseMovementMode> FindMovementMode(FName ModeName) const;
+		CHAOSMOVER_API void QueueLayeredMove(TSharedPtr<FLayeredMoveBase> Move);
 
-		void QueueLayeredMove(TSharedPtr<FLayeredMoveBase> Move);
+		CHAOSMOVER_API void QueueInstantMovementEffect(TSharedPtr<FInstantMovementEffect> Effect);
 
-		void QueueInstantMovementEffect(TSharedPtr<FInstantMovementEffect> Effect);
+		CHAOSMOVER_API FMovementModifierHandle QueueMovementModifier(TSharedPtr<FMovementModifierBase> Modifier);
+		CHAOSMOVER_API void CancelModifierFromHandle(FMovementModifierHandle ModifierHandle);
+		CHAOSMOVER_API const FMovementModifierBase* FindQueuedModifier(FMovementModifierHandle ModifierHandle) const;
+		CHAOSMOVER_API const FMovementModifierBase* FindQueuedModifierByType(const UScriptStruct* ModifierType) const;
 
 		void SetOwnerActorName(const FString& InOwnerActorName);
 		void SetOwnerActorLocalNetRole(ENetRole InOwnerActorLocalNetRole);
@@ -85,12 +89,25 @@ namespace UE::ChaosMover
 		/** Effects that are queued to be applied to the simulation at the start of the next sim subtick or at the end of this tick */
 		TArray<TSharedPtr<FInstantMovementEffect>> QueuedInstantEffects;
 
+		/** Modifiers that are queued to be added to the simulation at the start of the next sim subtick. */
+		TArray<TSharedPtr<FMovementModifierBase>> QueuedMovementModifiers;
+
+		/** Modifiers that are to be canceled at the start of the next sim subtick. */
+		TArray<FMovementModifierHandle> ModifiersToCancel;
+
+		// Internal-use-only tick data structs, for efficiency since they typically have the same contents from frame to frame
+		FMoverTickStartData WorkingSubstepStartData;
+		FSimulationTickParams WorkingSimTickParams;
+
 		TWeakObjectPtr<UChaosMoverSimulation> Simulation;
 
 	private:
 		void ConstructDefaultModes();
 		void AdvanceToNextMode();
 		void FlushQueuedMovesToGroup(FLayeredMoveGroup& Group);
+		void FlushQueuedModifiersToGroup(FMovementModifierGroup& ModifierGroup);
+		void FlushModifierCancellationsToGroup(FMovementModifierGroup& ActiveModifierGroup);
+		bool HasAnyInstantEffectsQueued() const;
 		bool ApplyInstantEffects(const FMoverTickStartData& SubstepStartData, const FMoverTimeStep& SubTimeStep, FMoverSyncState& OutputState);
 	
 		float InternalSimTimeMs = 0.0f;
diff --git a/Engine/Plugins/Experimental/ChaosMover/Source/ChaosMover/Private/Character/ChaosCharacterInputs.cpp b/Engine/Plugins/Experimental/ChaosMover/Source/ChaosMover/Private/Character/ChaosCharacterInputs.cpp
--- a/Engine/Plugins/Experimental/ChaosMover/Source/ChaosMover/Private/Character/ChaosCharacterInputs.cpp
+++ b/Engine/Plugins/Experimental/ChaosMover/Source/ChaosMover/Private/Character/ChaosCharacterInputs.cpp
@@ -42,4 +42,126 @@ void FChaosMoverLaunchInputs::Interpolate(const FMoverDataStructBase& From, cons
 void FChaosMoverLaunchInputs::Merge(const FMoverDataStructBase& From)
 {
 
+}
+
+bool FChaosMoverCrouchInputs::NetSerialize(FArchive& Ar, UPackageMap* Map, bool& bOutSuccess)
+{
+	Super::NetSerialize(Ar, Map, bOutSuccess);
+
+	Ar << bWantsToCrouch;
+
+	bOutSuccess = true;
+	return true;
+}
+
+void FChaosMoverCrouchInputs::ToString(FAnsiStringBuilderBase& Out) const
+{
+	Super::ToString(Out);
+
+	Out.Appendf("bWantsToCrouch: %hs \n", *LexToString(bWantsToCrouch));
+}
+
+bool FChaosMoverCrouchInputs::ShouldReconcile(const FMoverDataStructBase& AuthorityState) const
+{
+	const FChaosMoverCrouchInputs& TypedAuthority = static_cast<const FChaosMoverCrouchInputs&>(AuthorityState);
+	return bWantsToCrouch != TypedAuthority.bWantsToCrouch;
+}
+
+void FChaosMoverCrouchInputs::Interpolate(const FMoverDataStructBase& From, const FMoverDataStructBase& To, float Pct)
+{
+	if (Pct < .5f)
+	{
+		*this = static_cast<const FChaosMoverCrouchInputs&>(From);
+	}
+	else
+	{
+		*this = static_cast<const FChaosMoverCrouchInputs&>(To);
+	}
+}
+
+void FChaosMoverCrouchInputs::Merge(const FMoverDataStructBase& From)
+{
+}
+
+bool FChaosMovementSettingsOverrides::NetSerialize(FArchive& Ar, UPackageMap* Map, bool& bOutSuccess)
+{
+	Super::NetSerialize(Ar, Map, bOutSuccess);
+
+	Ar << ModeName;
+	Ar << MaxSpeedOverride;
+	Ar << AccelerationOverride;
+
+	bOutSuccess = true;
+	return true;
+}
+
+void FChaosMovementSettingsOverrides::ToString(FAnsiStringBuilderBase& Out) const
+{
+	Super::ToString(Out);
+
+	Out.Appendf("OverrideModeName: %s |", *ModeName.ToString());
+	Out.Appendf("MaxSpeedOverride: %.2f |", MaxSpeedOverride);
+	Out.Appendf("MaxSpeedOverride: %.2f \n", AccelerationOverride);
+}
+
+bool FChaosMovementSettingsOverrides::ShouldReconcile(const FMoverDataStructBase& AuthorityState) const
+{
+	const FChaosMovementSettingsOverrides& TypedAuthority = static_cast<const FChaosMovementSettingsOverrides&>(AuthorityState);
+	return ModeName != TypedAuthority.ModeName || !FMath::IsNearlyEqual(MaxSpeedOverride, TypedAuthority.MaxSpeedOverride) || !FMath::IsNearlyEqual(AccelerationOverride, TypedAuthority.AccelerationOverride);
+}
+
+void FChaosMovementSettingsOverrides::Interpolate(const FMoverDataStructBase& From, const FMoverDataStructBase& To, float Pct)
+{
+	const FChaosMovementSettingsOverrides& TypedFrom = static_cast<const FChaosMovementSettingsOverrides&>(From);
+	const FChaosMovementSettingsOverrides& TypedTo = static_cast<const FChaosMovementSettingsOverrides&>(To);
+	
+	ModeName = TypedTo.ModeName;
+	if (TypedFrom.ModeName == TypedTo.ModeName)
+	{
+		MaxSpeedOverride = FMath::Lerp(TypedFrom.MaxSpeedOverride, TypedTo.MaxSpeedOverride, Pct);
+		AccelerationOverride = FMath::Lerp(TypedFrom.AccelerationOverride, TypedTo.AccelerationOverride, Pct);
+	}
+	else
+	{
+		MaxSpeedOverride = TypedTo.MaxSpeedOverride;
+		AccelerationOverride = TypedTo.AccelerationOverride;
+	}
+
+}
+
+void FChaosMovementSettingsOverrides::Merge(const FMoverDataStructBase& From)
+{
+}
+
+bool FChaosMovementSettingsOverridesRemover::NetSerialize(FArchive& Ar, UPackageMap* Map, bool& bOutSuccess)
+{
+	Super::NetSerialize(Ar, Map, bOutSuccess);
+
+	Ar << ModeName;
+
+	bOutSuccess = true;
+	return true;
+}
+
+void FChaosMovementSettingsOverridesRemover::ToString(FAnsiStringBuilderBase& Out) const
+{
+	Super::ToString(Out);
+
+	Out.Appendf("eModeName: %s |", *ModeName.ToString());
+}
+
+bool FChaosMovementSettingsOverridesRemover::ShouldReconcile(const FMoverDataStructBase& AuthorityState) const
+{
+	const FChaosMovementSettingsOverridesRemover& TypedAuthority = static_cast<const FChaosMovementSettingsOverridesRemover&>(AuthorityState);
+	return ModeName != TypedAuthority.ModeName;
+}
+
+void FChaosMovementSettingsOverridesRemover::Interpolate(const FMoverDataStructBase& From, const FMoverDataStructBase& To, float Pct)
+{
+	const FChaosMovementSettingsOverridesRemover& TypedTo = static_cast<const FChaosMovementSettingsOverridesRemover&>(To);
+	ModeName = TypedTo.ModeName;
+}
+
+void FChaosMovementSettingsOverridesRemover::Merge(const FMoverDataStructBase& From)
+{
 }
\ No newline at end of file
diff --git a/Engine/Plugins/Experimental/ChaosMover/Source/ChaosMover/Private/Character/ChaosCharacterMoverComponent.cpp b/Engine/Plugins/Experimental/ChaosMover/Source/ChaosMover/Private/Character/ChaosCharacterMoverComponent.cpp
--- a/Engine/Plugins/Experimental/ChaosMover/Source/ChaosMover/Private/Character/ChaosCharacterMoverComponent.cpp
+++ b/Engine/Plugins/Experimental/ChaosMover/Source/ChaosMover/Private/Character/ChaosCharacterMoverComponent.cpp
@@ -9,6 +9,7 @@
 #include "ChaosMover/Character/Modes/ChaosWalkingMode.h"
 #include "ChaosMover/ChaosMoverLog.h"
 #include "DefaultMovementSet/CharacterMoverSimulationTypes.h"
+#include "GameFramework/Pawn.h"
 
 UChaosCharacterMoverComponent::UChaosCharacterMoverComponent()
 {
@@ -25,6 +26,12 @@ UChaosCharacterMoverComponent::UChaosCharacterMoverComponent()
 	BackendClass = UChaosMoverBackendComponent::StaticClass();
 }
 
+void UChaosCharacterMoverComponent::OnMoverPreSimulationTick(const FMoverTimeStep& TimeStep, const FMoverInputCmdContext& InputCmd)
+{
+	// This overrides the base class to make sure that it doesn't do anything. We do this
+	// because the stance modifier and jump handling of the base is not async physics friendly
+}
+
 void UChaosCharacterMoverComponent::ProcessSimulationEvent(const FMoverSimulationEventData& EventData)
 {
 	Super::ProcessSimulationEvent(EventData);
@@ -33,6 +40,10 @@ void UChaosCharacterMoverComponent::ProcessSimulationEvent(const FMoverSimulatio
 	{
 		OnLanded.Broadcast(LandedData->NewModeName, LandedData->HitResult);
 	}
+	if (const FStanceModifiedEventData* StanceModifiedEvent = EventData.CastTo<FStanceModifiedEventData>())
+	{
+		OnStanceModified(*StanceModifiedEvent);
+	}
 }
 
 void UChaosCharacterMoverComponent::DoQueueNextMode(FName DesiredModeName, bool bShouldReenter)
@@ -44,7 +55,7 @@ void UChaosCharacterMoverComponent::DoQueueNextMode(FName DesiredModeName, bool
 		const FName NextModeName = QueuedModeTransitionName;
 		if (NextModeName != NAME_None)
 		{
-			UE_LOG(LogMover, Log, TEXT("%s (%s) Overwriting of queued mode change (%s) with (%s)"), *GetNameSafe(GetOwner()), *UEnum::GetValueAsString(GetOwner()->GetLocalRole()), *NextModeName.ToString(), *DesiredModeName.ToString());
+			UE_LOG(LogChaosMover, Log, TEXT("%s (%s) Overwriting of queued mode change (%s) with (%s)"), *GetNameSafe(GetOwner()), *UEnum::GetValueAsString(GetOwner()->GetLocalRole()), *NextModeName.ToString(), *DesiredModeName.ToString());
 		}
 	}
 
@@ -103,10 +114,70 @@ void UChaosCharacterMoverComponent::ProduceInput(const int32 DeltaTimeMS, FMover
 
 		ClearQueuedMode();
 	}
+
+	if (bCrouchInputPending)
+	{
+		if ((bWantsToCrouch && CanCrouch()) || !bWantsToCrouch)
+		{
+			FChaosMoverCrouchInputs& CrouchInputs = Cmd->InputCollection.FindOrAddMutableDataByType<FChaosMoverCrouchInputs>();
+			CrouchInputs.bWantsToCrouch = bWantsToCrouch;
+		}
+
+		bCrouchInputPending = false;
+	}
+
+	if (bCancelMovementOverrides)
+	{
+		FChaosMovementSettingsOverridesRemover& MovementSettingsRemover = Cmd->InputCollection.FindOrAddMutableDataByType<FChaosMovementSettingsOverridesRemover>();
+		MovementSettingsRemover.ModeName = ModeToOverrideSettings;
+
+		bCancelMovementOverrides = false;
+	}
+
+	if (bOverrideMovementSettings)
+	{
+		FChaosMovementSettingsOverrides& MovementSettings = Cmd->InputCollection.FindOrAddMutableDataByType<FChaosMovementSettingsOverrides>();
+		MovementSettings.MaxSpeedOverride = MaxSpeedOverride;
+		MovementSettings.AccelerationOverride = AccelerationOverride;
+		MovementSettings.ModeName = ModeToOverrideSettings;
+
+		bOverrideMovementSettings = false;
+	}
 }
 
 void UChaosCharacterMoverComponent::Launch(const FVector& VelocityOrImpulse, EChaosMoverVelocityEffectMode Mode)
 {
 	LaunchVelocityOrImpulse = VelocityOrImpulse;
 	LaunchMode = Mode;
+}
+
+void UChaosCharacterMoverComponent::Crouch()
+{
+	Super::Crouch();
+	bCrouchInputPending = true;
+}
+
+void UChaosCharacterMoverComponent::UnCrouch()
+{
+	Super::UnCrouch();
+	bCrouchInputPending = true;
+}
+
+void UChaosCharacterMoverComponent::OnStanceModified(const FStanceModifiedEventData& EventData)
+{
+	OnStanceChanged.Broadcast(EventData.OldStance, EventData.NewStance);
+}
+
+void UChaosCharacterMoverComponent::OverrideMovementSettings(const FChaosMovementSettingsOverrides Overrides)
+{
+	ModeToOverrideSettings = Overrides.ModeName;
+	MaxSpeedOverride = FMath::Max(Overrides.MaxSpeedOverride, 0.0f);
+	AccelerationOverride = FMath::Max(Overrides.AccelerationOverride, 0.0f);
+	bOverrideMovementSettings = true;
+}
+
+void UChaosCharacterMoverComponent::CancelMovementSettingsOverrides(FName ModeName)
+{
+	ModeToOverrideSettings = ModeName;
+	bCancelMovementOverrides = true;
 }
\ No newline at end of file
diff --git a/Engine/Plugins/Experimental/ChaosMover/Source/ChaosMover/Private/Character/Modes/ChaosCharacterMovementMode.cpp b/Engine/Plugins/Experimental/ChaosMover/Source/ChaosMover/Private/Character/Modes/ChaosCharacterMovementMode.cpp
--- a/Engine/Plugins/Experimental/ChaosMover/Source/ChaosMover/Private/Character/Modes/ChaosCharacterMovementMode.cpp
+++ b/Engine/Plugins/Experimental/ChaosMover/Source/ChaosMover/Private/Character/Modes/ChaosCharacterMovementMode.cpp
@@ -4,6 +4,7 @@
 
 #include "Chaos/Character/CharacterGroundConstraintSettings.h"
 #include "ChaosMover/Character/Settings/SharedChaosCharacterMovementSettings.h"
+#include "ChaosMover/ChaosMoverLog.h"
 #include "Components/CapsuleComponent.h"
 #include "Components/SkeletalMeshComponent.h"
 #include "GameFramework/Character.h"
@@ -100,3 +101,53 @@ float UChaosCharacterMovementMode::GetMaxWalkSlopeCosine() const
 
 	return 0.707f;
 }
+
+float UChaosCharacterMovementMode::GetMaxSpeed() const
+{
+	if (MaxSpeedOverride.IsSet())
+	{
+		return MaxSpeedOverride.GetValue();
+	}
+	else if (const USharedChaosCharacterMovementSettings* SharedSettingsPtr = SharedSettings.Get())
+	{
+		return SharedSettingsPtr->MaxSpeed;
+	}
+
+	UE_LOG(LogChaosMover, Warning, TEXT("Invalid max speed on ChaosCharacterMoverComponent"));
+	return 0.0f;
+}
+
+void UChaosCharacterMovementMode::OverrideMaxSpeed(float Value)
+{
+	MaxSpeedOverride = Value;
+}
+
+void UChaosCharacterMovementMode::ClearMaxSpeedOverride()
+{
+	MaxSpeedOverride.Reset();
+}
+
+float UChaosCharacterMovementMode::GetAcceleration() const
+{
+	if (AccelerationOverride.IsSet())
+	{
+		return AccelerationOverride.GetValue();
+	}
+	else if (const USharedChaosCharacterMovementSettings* SharedSettingsPtr = SharedSettings.Get())
+	{
+		return SharedSettingsPtr->Acceleration;
+	}
+	
+	UE_LOG(LogChaosMover, Warning, TEXT("Invalid acceleration on ChaosCharacterMoverComponent"));
+	return 0.0f;
+}
+
+void UChaosCharacterMovementMode::OverrideAcceleration(float Value)
+{
+	AccelerationOverride = Value;
+}
+
+void UChaosCharacterMovementMode::ClearAccelerationOverride()
+{
+	AccelerationOverride.Reset();
+}
diff --git a/Engine/Plugins/Experimental/ChaosMover/Source/ChaosMover/Private/Character/Modes/ChaosFallingMode.cpp b/Engine/Plugins/Experimental/ChaosMover/Source/ChaosMover/Private/Character/Modes/ChaosFallingMode.cpp
--- a/Engine/Plugins/Experimental/ChaosMover/Source/ChaosMover/Private/Character/Modes/ChaosFallingMode.cpp
+++ b/Engine/Plugins/Experimental/ChaosMover/Source/ChaosMover/Private/Character/Modes/ChaosFallingMode.cpp
@@ -92,8 +92,8 @@ void UChaosFallingMode::GenerateMove_Implementation(const FMoverTickStartData& S
 	Params.DeltaSeconds = DeltaSeconds;
 	Params.TurningRate = SharedSettings->TurningRate;
 	Params.TurningBoost = SharedSettings->TurningBoost;
-	Params.MaxSpeed = SharedSettings->MaxSpeed;
-	Params.Acceleration = SharedSettings->Acceleration;
+	Params.MaxSpeed = GetMaxSpeed();
+	Params.Acceleration = GetAcceleration();
 	Params.Deceleration = FallingDeceleration;
 	Params.WorldToGravityQuat = WorldToGravityTransform;
 	Params.bUseAccelerationForVelocityMove = SharedSettings->bUseAccelerationForVelocityMove;
@@ -129,16 +129,16 @@ void UChaosFallingMode::GenerateMove_Implementation(const FMoverTickStartData& S
 		SimBlackboard->Set(CommonBlackboard::LastFloorResult, FloorResult);
 		SimBlackboard->Set(CommonBlackboard::LastWaterResult, WaterResult);
 
-		if (FloorResult.HitResult.IsValidBlockingHit() && FloorResult.HitResult.Normal.Dot(DefaultSimInputs->UpDir) > UE::MoverUtils::VERTICAL_SLOPE_NORMAL_MAX_DOT && !FloorResult.IsWalkableFloor())
-		{
-			// If acceleration is into the wall, limit contribution.
-			if (FVector::DotProduct(Params.MoveInput, FloorResult.HitResult.Normal) < 0.f)
-			{
-				// Allow movement parallel to the wall, but not into it because that may push us up.
-				const FVector FallingHitNormal = FVector::VectorPlaneProject(FloorResult.HitResult.Normal, -DefaultSimInputs->UpDir).GetSafeNormal();
-				Params.MoveInput = FVector::VectorPlaneProject(Params.MoveInput, FallingHitNormal);
-			}
-		}
+		//if (FloorResult.HitResult.IsValidBlockingHit() && FloorResult.HitResult.Normal.Dot(DefaultSimInputs->UpDir) > UE::MoverUtils::VERTICAL_SLOPE_NORMAL_MAX_DOT && !FloorResult.IsWalkableFloor())
+		//{
+		//	// If acceleration is into the wall, limit contribution.
+		//	if (FVector::DotProduct(Params.MoveInput, FloorResult.HitResult.Normal) < 0.f)
+		//	{
+		//		// Allow movement parallel to the wall, but not into it because that may push us up.
+		//		const FVector FallingHitNormal = FVector::VectorPlaneProject(FloorResult.HitResult.Normal, -DefaultSimInputs->UpDir).GetSafeNormal();
+		//		Params.MoveInput = FVector::VectorPlaneProject(Params.MoveInput, FallingHitNormal);
+		//	}
+		//}
 	}
 
 	OutProposedMove = UAirMovementUtils::ComputeControlledFreeMove(Params);
@@ -198,6 +198,7 @@ void UChaosFallingMode::SimulationTick_Implementation(const FSimulationTickParam
 	FVector TargetPos = StartingSyncState->GetLocation_WorldSpace() + TargetVel * DeltaSeconds;
 
 	OutputState.MovementEndState.RemainingMs = 0.0f;
+	OutputState.MovementEndState.NextModeName = Params.StartState.SyncState.MovementMode;
 	OutputSyncState.MoveDirectionIntent = ProposedMove.bHasDirIntent ? ProposedMove.DirectionIntent : FVector::ZeroVector;
 	OutputSyncState.SetTransforms_WorldSpace(
 		TargetPos,
diff --git a/Engine/Plugins/Experimental/ChaosMover/Source/ChaosMover/Private/Character/Modes/ChaosFlyingMode.cpp b/Engine/Plugins/Experimental/ChaosMover/Source/ChaosMover/Private/Character/Modes/ChaosFlyingMode.cpp
--- a/Engine/Plugins/Experimental/ChaosMover/Source/ChaosMover/Private/Character/Modes/ChaosFlyingMode.cpp
+++ b/Engine/Plugins/Experimental/ChaosMover/Source/ChaosMover/Private/Character/Modes/ChaosFlyingMode.cpp
@@ -79,8 +79,8 @@ void UChaosFlyingMode::GenerateMove_Implementation(const FMoverTickStartData& St
 	Params.PriorOrientation = StartingSyncState->GetOrientation_WorldSpace();
 	Params.TurningRate = SharedSettings->TurningRate;
 	Params.TurningBoost = SharedSettings->TurningBoost;
-	Params.MaxSpeed = SharedSettings->MaxSpeed;
-	Params.Acceleration = SharedSettings->Acceleration;
+	Params.MaxSpeed = GetMaxSpeed();
+	Params.Acceleration = GetAcceleration();
 	Params.Deceleration = SharedSettings->Deceleration;
 	Params.DeltaSeconds = DeltaSeconds;
 	Params.WorldToGravityQuat = WorldToGravityTransform;
@@ -127,6 +127,7 @@ void UChaosFlyingMode::SimulationTick_Implementation(const FSimulationTickParams
 	FVector TargetPos = StartingSyncState->GetLocation_WorldSpace() + TargetVel * DeltaSeconds;
 
 	OutputState.MovementEndState.RemainingMs = 0.0f;
+	OutputState.MovementEndState.NextModeName = Params.StartState.SyncState.MovementMode;
 	OutputSyncState.MoveDirectionIntent = ProposedMove.bHasDirIntent ? ProposedMove.DirectionIntent : FVector::ZeroVector;
 	OutputSyncState.SetTransforms_WorldSpace(
 		TargetPos,
diff --git a/Engine/Plugins/Experimental/ChaosMover/Source/ChaosMover/Private/Character/Modes/ChaosSwimmingMode.cpp b/Engine/Plugins/Experimental/ChaosMover/Source/ChaosMover/Private/Character/Modes/ChaosSwimmingMode.cpp
--- a/Engine/Plugins/Experimental/ChaosMover/Source/ChaosMover/Private/Character/Modes/ChaosSwimmingMode.cpp
+++ b/Engine/Plugins/Experimental/ChaosMover/Source/ChaosMover/Private/Character/Modes/ChaosSwimmingMode.cpp
@@ -198,8 +198,8 @@ void UChaosSwimmingMode::GenerateMove_Implementation(const FMoverTickStartData&
 		Params.PriorOrientation = StartingSyncState->GetOrientation_WorldSpace();
 		Params.TurningRate = SharedSettings->TurningRate;
 		Params.TurningBoost = SharedSettings->TurningBoost;
-		Params.MaxSpeed = SharedSettings->MaxSpeed;
-		Params.Acceleration = SharedSettings->Acceleration;
+		Params.MaxSpeed = GetMaxSpeed();
+		Params.Acceleration = GetAcceleration();
 		Params.Deceleration = SharedSettings->Deceleration;
 		Params.DeltaSeconds = DeltaSeconds;
 		Params.MoveSpeed = WaterSpeed;
@@ -256,4 +256,5 @@ void UChaosSwimmingMode::SimulationTick_Implementation(const FSimulationTickPara
 		TargetVel);
 
 	OutputState.MovementEndState.RemainingMs = 0.0f;
+	OutputState.MovementEndState.NextModeName = Params.StartState.SyncState.MovementMode;
 }
\ No newline at end of file
diff --git a/Engine/Plugins/Experimental/ChaosMover/Source/ChaosMover/Private/Character/Modes/ChaosWalkingMode.cpp b/Engine/Plugins/Experimental/ChaosMover/Source/ChaosMover/Private/Character/Modes/ChaosWalkingMode.cpp
--- a/Engine/Plugins/Experimental/ChaosMover/Source/ChaosMover/Private/Character/Modes/ChaosWalkingMode.cpp
+++ b/Engine/Plugins/Experimental/ChaosMover/Source/ChaosMover/Private/Character/Modes/ChaosWalkingMode.cpp
@@ -120,8 +120,8 @@ void UChaosWalkingMode::GenerateMove_Implementation(const FMoverTickStartData& S
 	Params.GroundNormal = MovementNormal;
 	Params.TurningRate = SharedSettings->TurningRate;
 	Params.TurningBoost = SharedSettings->TurningBoost;
-	Params.MaxSpeed = SharedSettings->MaxSpeed;
-	Params.Acceleration = SharedSettings->Acceleration;
+	Params.MaxSpeed = GetMaxSpeed();
+	Params.Acceleration = GetAcceleration();
 	Params.Deceleration = SharedSettings->Deceleration;
 	Params.DeltaSeconds = DeltaSeconds;
 	Params.WorldToGravityQuat = WorldToGravityTransform;
@@ -277,6 +277,7 @@ void UChaosWalkingMode::SimulationTick_Implementation(const FSimulationTickParam
 	}
 	
 	OutputState.MovementEndState.RemainingMs = 0.0f;
+	OutputState.MovementEndState.NextModeName = Params.StartState.SyncState.MovementMode;
 	OutputSyncState.MoveDirectionIntent = ProposedMove.bHasDirIntent ? ProposedMove.DirectionIntent : FVector::ZeroVector;
 }
 
diff --git a/Engine/Plugins/Experimental/ChaosMover/Source/ChaosMover/Private/Character/Modifiers/ChaosStanceModifier.cpp b/Engine/Plugins/Experimental/ChaosMover/Source/ChaosMover/Private/Character/Modifiers/ChaosStanceModifier.cpp
new file mode 100644
--- /dev/null
+++ b/Engine/Plugins/Experimental/ChaosMover/Source/ChaosMover/Private/Character/Modifiers/ChaosStanceModifier.cpp
@@ -0,0 +1,250 @@
+// Copyright Epic Games, Inc. All Rights Reserved.
+
+#include "ChaosMover/Character/Modifiers/ChaosStanceModifier.h"
+
+#include "Chaos/Capsule.h"
+#include "Chaos/PhysicsObjectInternalInterface.h"
+#include "ChaosMover/ChaosMoverLog.h"
+#include "ChaosMover/ChaosMoverSimulation.h"
+#include "Components/CapsuleComponent.h"
+#include "MoverTypes.h"
+#include "PBDRigidsSolver.h"
+#include "Physics/Experimental/PhysScene_Chaos.h"
+#include "PhysicsEngine/PhysicsObjectExternalInterface.h"
+#include "Engine/World.h"
+
+FChaosStanceModifier::FChaosStanceModifier()
+	: ActiveStance(EStanceMode::Invalid)
+{
+	DurationMs = -1.0f;
+}
+
+bool FChaosStanceModifier::HasGameplayTag(FGameplayTag TagToFind, bool bExactMatch) const
+{
+	// TODO: Eventually check for other stance tags here like prone
+	if (bExactMatch)
+	{
+		return TagToFind.MatchesTagExact(Mover_IsCrouching);
+	}
+
+	return TagToFind.MatchesTag(Mover_IsCrouching);
+}
+
+const Chaos::FCapsule* FChaosStanceModifier::GetCapsule(Chaos::FPBDRigidParticleHandle* ParticleHandle) const
+{
+	const Chaos::FCapsule* FoundCapsule = nullptr;
+
+	const Chaos::FImplicitObjectRef Geometry = ParticleHandle->GetGeometry();
+	if (Geometry && Geometry->IsValidGeometry())
+	{
+		FoundCapsule = ParticleHandle->GetGeometry()->GetObject<Chaos::FCapsule>();
+		if (!FoundCapsule)
+		{
+			if (const Chaos::FImplicitObjectUnion* ImplicitUnion = ParticleHandle->GetGeometry()->GetObject<Chaos::FImplicitObjectUnion>())
+			{
+				for (const Chaos::FImplicitObjectRef ObjectRef : ImplicitUnion->GetObjects())
+				{
+					FoundCapsule = ObjectRef->GetObject<Chaos::FCapsule>();
+					if (FoundCapsule)
+					{
+						break;
+					}
+				}
+			}
+		}
+	}
+
+	return FoundCapsule;
+}
+
+void FChaosStanceModifier::UpdateStance(const FMovementModifierParams_Async& AsyncParams, EStanceMode NewStance)
+{
+	if (NewStance == ActiveStance)
+	{
+		UE_LOG(LogChaosMover, Verbose, TEXT("Skipping stance change - not required"));
+		return;
+	}
+
+	if (!AsyncParams.IsValid())
+	{
+		return;
+	}
+
+	UChaosMoverSimulation* Sim = Cast<UChaosMoverSimulation>(AsyncParams.Simulation);
+	if (!Sim)
+	{
+		return;
+	}
+
+	FName ModeName = CurrentModeName.IsNone() ? AsyncParams.SyncState->MovementMode : CurrentModeName;
+	IChaosCharacterMovementModeInterface* CharacterMode = Cast<IChaosCharacterMovementModeInterface>(Sim->FindMovementModeByName_Mutable(ModeName));
+	if (!CharacterMode)
+	{
+		UE_LOG(LogChaosMover, Warning, TEXT("ChaosStanceModifier only works with character modes"));
+		return;
+	}
+
+	const FChaosMoverSimulationDefaultInputs* SimInputs = Sim->GetLocalSimInput().FindDataByType<FChaosMoverSimulationDefaultInputs>();
+	if (!SimInputs)
+	{
+		UE_LOG(LogChaosMover, Warning, TEXT("ChaosStanceModifier requires FCharacterDefaultInputs"));
+		return;
+	}
+
+	Chaos::FPhysicsSolver* Solver = nullptr;
+	if (SimInputs->World)
+	{
+		if (FPhysScene* Scene = SimInputs->World->GetPhysicsScene())
+		{
+			Solver = Scene->GetSolver();
+		}
+	}
+
+	if (!Solver)
+	{
+		UE_LOG(LogChaosMover, Warning, TEXT("ChaosStanceModifier requires FPhysicsSolver"));
+		return;
+	}
+
+	Chaos::FWritePhysicsObjectInterface_Internal WriteInterface = Chaos::FPhysicsObjectInternalInterface::GetWrite();
+	Chaos::FPBDRigidParticleHandle* ParticleHandle = WriteInterface.GetRigidParticle(SimInputs->PhysicsObject);
+
+	if (!ParticleHandle)
+	{
+		return;
+	}
+
+	const Chaos::FCapsule* Capsule = GetCapsule(ParticleHandle);
+
+	if (!Capsule)
+	{
+		UE_LOG(LogChaosMover, Warning, TEXT("No capsule shape found in ChaosStanceModifier"));
+		return;
+	}
+
+	// Invalidate the particle as we are about to change the collision geometry so need to clear contacts etc.
+	Solver->GetEvolution()->InvalidateParticle(ParticleHandle);
+
+
+	// Scale the capsule
+	float NewHalfHeight, NewRadius, NewGroundClearance;
+	if (NewStance == EStanceMode::Crouch)
+	{
+		NewHalfHeight = ModifiedCapsuleHalfHeight;
+		NewRadius = ModifiedCapsuleRadius;
+		NewGroundClearance = ModifiedCapsuleGroundClearance;
+
+		if (MaxSpeedOverride.IsSet())
+		{
+			CharacterMode->OverrideMaxSpeed(MaxSpeedOverride.GetValue());
+		}
+
+		if (AccelerationOverride.IsSet())
+		{
+			CharacterMode->OverrideAcceleration(AccelerationOverride.GetValue());
+		}
+	}
+	else
+	{
+		NewHalfHeight = DefaultCapsuleHalfHeight;
+		NewRadius = DefaultCapsuleRadius;
+		NewGroundClearance = DefaultCapsuleGroundClearance;
+
+		CharacterMode->ClearMaxSpeedOverride();
+		CharacterMode->ClearAccelerationOverride();
+	}
+	
+	UpdateCapsule(ParticleHandle, NewRadius, NewHalfHeight, NewGroundClearance, CharacterMode->GetTargetHeight());
+
+	EStanceMode OldStance = ActiveStance;
+	ActiveStance = NewStance;
+
+	// We're setting the capsule on PT but there is no automatic mechanism to
+	// reflect that change on GT, so add a callback to do that
+	Chaos::FImplicitObjectPtr NewCapsule = ParticleHandle->GetGeometry();
+	Chaos::FPhysicsObjectHandle PhysicsObject = SimInputs->PhysicsObject;
+
+	auto SetGeometryGTCallback = [NewCapsule, PhysicsObject](const FMoverSimulationEventData & Data) {
+		Chaos::FPBDRigidParticle* Particle = FPhysicsObjectExternalInterface::LockRead(PhysicsObject)->GetRigidParticle(PhysicsObject);
+		if (Particle && NewCapsule.IsValid())
+		{
+			Particle->SetGeometry(NewCapsule);
+		}
+	};
+
+	// Fire simulation event
+	Sim->AddEvent(MakeShared<FStanceModifiedEventData>(AsyncParams.TimeStep->BaseSimTimeMs, OldStance, ActiveStance, SetGeometryGTCallback));
+}
+
+void FChaosStanceModifier::OnStart_Async(const FMovementModifierParams_Async& AsyncParams)
+{
+	UpdateStance(AsyncParams, EStanceMode::Crouch);
+
+	// Set the current mode - this is used to check if the mode changes
+	CurrentModeName = AsyncParams.SyncState->MovementMode;
+}
+
+void FChaosStanceModifier::OnEnd_Async(const FMovementModifierParams_Async& AsyncParams)
+{
+	UpdateStance(AsyncParams, EStanceMode::Invalid);
+
+	CurrentModeName = NAME_None;
+}
+
+void FChaosStanceModifier::UpdateCapsule(Chaos::FPBDRigidParticleHandle* ParticleHandle, float NewRadius, float NewHalfHeight, float NewGroundClearance, float TargetHeight)
+{
+	check(ParticleHandle);
+
+	const float NewHeight = FMath::Max(2.0f * (NewHalfHeight - NewRadius), 0.0f);
+	const Chaos::FVec3 Axis = Chaos::FVec3::UpVector;
+	const Chaos::FVec3 NewX1 = (NewGroundClearance + NewRadius - TargetHeight) * Axis;
+	const Chaos::FVec3 NewX2 = NewX1 + Axis * NewHeight;
+	const Chaos::FCapsule NewCapsule(NewX1, NewX2, NewRadius);
+
+	ParticleHandle->SetGeometry(Chaos::FImplicitObjectPtr(new Chaos::FCapsule(NewX1, NewX2, NewRadius)));
+	ParticleHandle->SetCenterOfMass(NewCapsule.GetCenterf());
+}
+
+void FChaosStanceModifier::OnPostMovement_Async(const FMovementModifierParams_Async& AsyncParams)
+{
+	if (bCancelOnModeChange && CurrentModeName != AsyncParams.SyncState->MovementMode)
+	{
+		DurationMs = 0;
+	}
+}
+
+FMovementModifierBase* FChaosStanceModifier::Clone() const
+{
+	return new FChaosStanceModifier(*this);
+}
+
+void FChaosStanceModifier::NetSerialize(FArchive& Ar)
+{
+	Super::NetSerialize(Ar);
+
+	Ar << ActiveStance;
+	Ar << bCancelOnModeChange;
+	Ar << ModifiedCapsuleHalfHeight;
+	Ar << ModifiedCapsuleRadius;
+	Ar << ModifiedCapsuleGroundClearance;
+	Ar << DefaultCapsuleHalfHeight;
+	Ar << DefaultCapsuleRadius;
+	Ar << DefaultCapsuleGroundClearance;
+	Ar << AccelerationOverride;
+	Ar << MaxSpeedOverride;
+}
+
+UScriptStruct* FChaosStanceModifier::GetScriptStruct() const
+{
+	return FChaosStanceModifier::StaticStruct();
+}
+
+FString FChaosStanceModifier::ToSimpleString() const
+{
+	return FString::Printf(TEXT("Chaos Stance Modifier"));
+}
+
+void FChaosStanceModifier::AddReferencedObjects(class FReferenceCollector& Collector)
+{
+	Super::AddReferencedObjects(Collector);
+}
\ No newline at end of file
diff --git a/Engine/Plugins/Experimental/ChaosMover/Source/ChaosMover/Private/Character/Transitions/ChaosCharacterCrouchCheck.cpp b/Engine/Plugins/Experimental/ChaosMover/Source/ChaosMover/Private/Character/Transitions/ChaosCharacterCrouchCheck.cpp
new file mode 100644
--- /dev/null
+++ b/Engine/Plugins/Experimental/ChaosMover/Source/ChaosMover/Private/Character/Transitions/ChaosCharacterCrouchCheck.cpp
@@ -0,0 +1,186 @@
+// Copyright Epic Games, Inc. All Rights Reserved.
+
+#include "ChaosMover/Character/Transitions/ChaosCharacterCrouchCheck.h"
+
+#include "Chaos/Capsule.h"
+#include "Chaos/PhysicsObjectInternalInterface.h"
+#include "ChaosMover/ChaosMoverLog.h"
+#include "ChaosMover/ChaosMoverSimulation.h"
+#include "ChaosMover/ChaosMoverSimulationTypes.h"
+#include "ChaosMover/Character/ChaosCharacterInputs.h"
+#include "ChaosMover/Character/Modes/ChaosCharacterMovementMode.h"
+#include "ChaosMover/Character/Modifiers/ChaosStanceModifier.h"
+#include "ChaosMover/Utilities/ChaosMoverQueryUtils.h"
+#include "Components/CapsuleComponent.h"
+#include "MoverComponent.h"
+
+#include UE_INLINE_GENERATED_CPP_BY_NAME(ChaosCharacterCrouchCheck)
+
+UChaosCharacterCrouchCheck::UChaosCharacterCrouchCheck(const FObjectInitializer& ObjectInitializer)
+	: Super(ObjectInitializer)
+{
+	bSupportsAsync = true;
+	bAllowModeReentry = true;
+	bFirstSubStepOnly = true;
+}
+
+void UChaosCharacterCrouchCheck::OnRegistered()
+{
+	Super::OnRegistered();
+
+	if (const UPrimitiveComponent* PrimComp = Cast<UPrimitiveComponent>(GetMoverComponent()->GetUpdatedComponent()))
+	{
+		PrimComp->CalcBoundingCylinder(OriginalCapsuleRadius, OriginalCapsuleHalfHeight);
+	}
+}
+
+void UChaosCharacterCrouchCheck::OnUnregistered()
+{
+	Super::OnUnregistered();
+}
+
+bool UChaosCharacterCrouchCheck::CanCrouch(const FSimulationTickParams& Params) const
+{
+	return true;
+}
+
+bool UChaosCharacterCrouchCheck::CanUncrouch(const FSimulationTickParams& Params) const
+{
+	if (const FChaosMoverSimulationDefaultInputs* SimInputs = Simulation->GetLocalSimInput().FindDataByType<FChaosMoverSimulationDefaultInputs>())
+	{
+		if (const FMoverDefaultSyncState* StartingSyncState = Params.StartState.SyncState.SyncStateCollection.FindDataByType<FMoverDefaultSyncState>())
+		{
+			UE::ChaosMover::Utils::FCapsuleOverlapTestParams OverlapTestParams{
+				.ResponseParams = SimInputs->CollisionResponseParams,
+				.QueryParams = SimInputs->CollisionQueryParams,
+				.Location = StartingSyncState->GetLocation_WorldSpace(),
+				.UpDir = SimInputs->UpDir,
+				.World = SimInputs->World,
+				.CapsuleHalfHeight = OriginalCapsuleHalfHeight,
+				.CapsuleRadius = OriginalCapsuleRadius,
+				.CollisionChannel = SimInputs->CollisionChannel
+			};
+
+			return !UE::ChaosMover::Utils::CapsuleOverlapTest_Internal(OverlapTestParams);
+		}
+	}
+
+	return false;
+}
+
+void UChaosCharacterCrouchCheck::Crouch(const FSimulationTickParams& TickParams)
+{
+	UE_LOG(LogChaosMover, Verbose, TEXT("Crouch Activated"));
+
+	float OriginalCapsuleGroundClearance = 10.0f;
+	if (const IChaosCharacterMovementModeInterface* CharacterMode = Cast<IChaosCharacterMovementModeInterface>(Simulation->GetCurrentMovementMode()))
+	{
+		OriginalCapsuleGroundClearance = CharacterMode->GetTargetHeight() - OriginalCapsuleHalfHeight;
+	}
+
+	TSharedPtr<FChaosStanceModifier> StanceModifier = MakeShared<FChaosStanceModifier>();
+	StanceModifier->ModifiedCapsuleHalfHeight = CapsuleHalfHeight.IsSet() ? CapsuleHalfHeight.GetValue() : OriginalCapsuleHalfHeight;
+	StanceModifier->ModifiedCapsuleRadius = CapsuleRadius.IsSet() ? CapsuleRadius.GetValue() : OriginalCapsuleRadius;
+	StanceModifier->ModifiedCapsuleGroundClearance = CapsuleGroundClearance.IsSet() ? CapsuleGroundClearance.GetValue() : OriginalCapsuleGroundClearance;
+	StanceModifier->DefaultCapsuleHalfHeight = OriginalCapsuleHalfHeight;
+	StanceModifier->DefaultCapsuleRadius = OriginalCapsuleRadius;
+	StanceModifier->DefaultCapsuleGroundClearance = OriginalCapsuleGroundClearance;
+	StanceModifier->AccelerationOverride = Acceleration;
+	StanceModifier->MaxSpeedOverride = MaxSpeed;
+	StanceModifier->bCancelOnModeChange = bCancelOnModeChange;
+
+	CrouchModifierHandle = Simulation->QueueMovementModifier(MoveTemp(StanceModifier));
+}
+
+void UChaosCharacterCrouchCheck::Uncrouch(const FSimulationTickParams& TickParams)
+{
+	UE_LOG(LogChaosMover, Verbose, TEXT("Uncrouch Activated"));
+
+	check(Simulation);
+
+	if (CrouchModifierHandle.IsValid())
+	{
+		Simulation->CancelModifierFromHandle(CrouchModifierHandle);
+		CrouchModifierHandle.Invalidate();
+	}
+	else if (const FChaosStanceModifier* Modifier = Simulation->FindMovementModifierByType<FChaosStanceModifier>())
+	{
+		Simulation->CancelModifierFromHandle(Modifier->GetHandle());
+		CrouchModifierHandle.Invalidate();
+	}
+}
+
+FTransitionEvalResult UChaosCharacterCrouchCheck::Evaluate_Implementation(const FSimulationTickParams& TickParams) const
+{
+	FTransitionEvalResult EvalResult;
+
+	if (!Simulation)
+	{
+		UE_LOG(LogChaosMover, Warning, TEXT("No Simulation set on UChaosCharacterCrouchCheck"));
+		return EvalResult;
+	}
+
+	if (bTriggerCrouch || bTriggerUncrouch)
+	{
+		UE_LOG(LogChaosMover, Warning, TEXT("Trying to crouch/uncrouch in the same substep"));
+		return EvalResult;
+	}
+
+	const FChaosMoverCrouchInputs* CrouchInputs = TickParams.StartState.InputCmd.InputCollection.FindDataByType<FChaosMoverCrouchInputs>();
+	if (CrouchInputs)
+	{
+		if (const UBaseMovementMode* Mode = Simulation->FindMovementModeByName(TickParams.StartState.SyncState.MovementMode))
+		{
+			bool bIsCrouching = Simulation->HasGameplayTag(Mover_IsCrouching, true);
+			if (bIsCrouching && !CrouchModifierHandle.IsValid())
+			{
+				if (const FChaosStanceModifier* Modifier = Simulation->FindMovementModifierByType<FChaosStanceModifier>())
+				{
+					Simulation->CancelModifierFromHandle(Modifier->GetHandle());
+					bIsCrouching = false;
+				}
+			}
+
+			if (bIsCrouching && !CrouchInputs->bWantsToCrouch && CanUncrouch(TickParams))
+			{
+				bTriggerUncrouch = true;
+				if (TransitionToUncrouchedMode.IsSet())
+				{
+					EvalResult.NextMode = TransitionToUncrouchedMode.GetValue();
+				}
+				else
+				{
+					EvalResult.NextMode = TickParams.StartState.SyncState.MovementMode;
+				}
+			}
+			else if (!bIsCrouching && CrouchInputs->bWantsToCrouch && CanCrouch(TickParams))
+			{
+				bTriggerCrouch = true;
+				if (TransitionToCrouchedMode.IsSet())
+				{
+					EvalResult.NextMode = TransitionToCrouchedMode.GetValue();
+				}
+				else
+				{
+					EvalResult.NextMode = TickParams.StartState.SyncState.MovementMode;
+				}
+			}
+		}
+	}
+
+	return EvalResult;
+}
+
+void UChaosCharacterCrouchCheck::Trigger_Implementation(const FSimulationTickParams& TickParams)
+{
+	if (bTriggerCrouch)
+	{
+		Crouch(TickParams);
+	}
+	else if (bTriggerUncrouch)
+	{
+		Uncrouch(TickParams);
+	}
+
+	bTriggerUncrouch = bTriggerCrouch = false;
+}
\ No newline at end of file
diff --git a/Engine/Plugins/Experimental/ChaosMover/Source/ChaosMover/Private/Utilities/ChaosMoverQueryUtils.cpp b/Engine/Plugins/Experimental/ChaosMover/Source/ChaosMover/Private/Utilities/ChaosMoverQueryUtils.cpp
--- a/Engine/Plugins/Experimental/ChaosMover/Source/ChaosMover/Private/Utilities/ChaosMoverQueryUtils.cpp
+++ b/Engine/Plugins/Experimental/ChaosMover/Source/ChaosMover/Private/Utilities/ChaosMoverQueryUtils.cpp
@@ -150,4 +150,44 @@ namespace UE::ChaosMover::Utils
 		Internal::GetWaterResultFromHitResults_Internal(Hits, Params.Location, Params.TargetHeight, OutWaterResult);
 	}
 
+	bool CapsuleOverlapTest_Internal(const FCapsuleOverlapTestParams& Params)
+	{
+		Chaos::EnsureIsInPhysicsThreadContext();
+
+		TArray<FHitResult> Hits;
+
+		FVector Start = Params.Location - (Params.CapsuleHalfHeight -  Params.CapsuleRadius) * Params.UpDir;
+		FVector End = Params.Location + (Params.CapsuleHalfHeight - Params.CapsuleRadius) * Params.UpDir;
+		FHitResult OutHit;
+		if (Chaos::Private::FGenericPhysicsInterface_Internal::SpherecastMulti(Params.World, Params.CapsuleRadius, Hits, Start, End, Params.CollisionChannel, Params.QueryParams, Params.ResponseParams))
+		{
+			OutHit = Hits.Last();
+		}
+
+		bool bBlockingHitFound = false;
+		if (OutHit.bBlockingHit)
+		{
+			if (!Params.bBlockOnDynamic)
+			{
+				Chaos::FReadPhysicsObjectInterface_Internal Interface = Chaos::FPhysicsObjectInternalInterface::GetRead();
+				bBlockingHitFound = !Interface.AreAllDynamic({ OutHit.PhysicsObject });
+			}
+			else
+			{
+				bBlockingHitFound = true;
+			}
+		}
+
+#if CHAOSMOVER_DEBUG_DRAW
+		if (CVars::bDrawOverlapQueries)
+		{
+			const FVector Center = Params.Location;
+			const FColor Color = bBlockingHitFound ? FColor::Purple : FColor::Turquoise;
+			Chaos::FDebugDrawQueue::GetInstance().DrawDebugCapsule(Center, Params.CapsuleHalfHeight, Params.CapsuleRadius, FQuat::Identity, Color, false, -1.f, 10, 1.0f);
+		}
+#endif
+
+		return bBlockingHitFound;
+	}
+
 } // namespace UE::ChaosMover::Utils
\ No newline at end of file
diff --git a/Engine/Plugins/Experimental/ChaosMover/Source/ChaosMover/Public/ChaosMover/ChaosMoverDeveloperSettings.h b/Engine/Plugins/Experimental/ChaosMover/Source/ChaosMover/Public/ChaosMover/ChaosMoverDeveloperSettings.h
--- a/Engine/Plugins/Experimental/ChaosMover/Source/ChaosMover/Public/ChaosMover/ChaosMoverDeveloperSettings.h
+++ b/Engine/Plugins/Experimental/ChaosMover/Source/ChaosMover/Public/ChaosMover/ChaosMoverDeveloperSettings.h
@@ -11,6 +11,7 @@ namespace UE::ChaosMover
 		CHAOSMOVER_API extern bool bForceSingleThreadedGT;
 		CHAOSMOVER_API extern bool bForceSingleThreadedPT;
 		CHAOSMOVER_API extern bool bDrawGroundQueries;
+		CHAOSMOVER_API extern bool bDrawOverlapQueries;
 		CHAOSMOVER_API extern bool bSkipGenerateMoveIfOverridden;
 	}
 };
diff --git a/Engine/Plugins/Experimental/ChaosMover/Source/ChaosMover/Public/ChaosMover/ChaosMoverSimulation.h b/Engine/Plugins/Experimental/ChaosMover/Source/ChaosMover/Public/ChaosMover/ChaosMoverSimulation.h
--- a/Engine/Plugins/Experimental/ChaosMover/Source/ChaosMover/Public/ChaosMover/ChaosMoverSimulation.h
+++ b/Engine/Plugins/Experimental/ChaosMover/Source/ChaosMover/Public/ChaosMover/ChaosMoverSimulation.h
@@ -45,6 +45,9 @@ public:
 	UFUNCTION(BlueprintPure, Category = Mover)
 	CHAOSMOVER_API const UBaseMovementMode* FindMovementModeByName(const FName& Name) const;
 
+	UFUNCTION(BlueprintPure, Category = Mover)
+	CHAOSMOVER_API UBaseMovementMode* FindMovementModeByName_Mutable(const FName& Name);
+
 	UFUNCTION(BlueprintCallable, CustomThunk, Category = Mover, meta = (CustomStructureParam = "EffectAsRawData", AllowAbstract = "false", DisplayName = "Queue Instant Movement Effect"))
 	void K2_QueueInstantMovementEffect(UPARAM(DisplayName = "Instant Movement Effect") const int32& EffectAsRawData);
 	DECLARE_FUNCTION(execK2_QueueInstantMovementEffect);
@@ -52,6 +55,26 @@ public:
 	// Queue a Instant Movement Effect to take place at the end of this frame or start of the next subtick - whichever happens first
 	CHAOSMOVER_API void QueueInstantMovementEffect(TSharedPtr<FInstantMovementEffect> Move);
 
+	UFUNCTION(BlueprintCallable, CustomThunk, Category = Mover, meta = (CustomStructureParam = "MoveAsRawData", AllowAbstract = "false", DisplayName = "Queue Movement Modifier"))
+	CHAOSMOVER_API FMovementModifierHandle K2_QueueMovementModifier(UPARAM(DisplayName = "Movement Modifier") const int32& MoveAsRawData);
+	DECLARE_FUNCTION(execK2_QueueMovementModifier);
+
+	// Queue a Movement Modifier to start during the next simulation frame.
+	CHAOSMOVER_API FMovementModifierHandle QueueMovementModifier(TSharedPtr<FMovementModifierBase> Modifier);
+
+	// Find movement modifier by type (returns the first modifier it finds). Returns nullptr if the modifier couldn't be found
+	CHAOSMOVER_API const FMovementModifierBase* FindMovementModifierByType(const UScriptStruct* DataStructType) const;
+
+	/** Find a movement modifier of a specific type in this components movement modifiers. If not found, null will be returned. */
+	template <typename ModifierT = FMovementModifierBase UE_REQUIRES(std::is_base_of_v<FMovementModifierBase, ModifierT>)>
+	CHAOSMOVER_API const ModifierT* FindMovementModifierByType() const { return static_cast<const ModifierT*>(FindMovementModifierByType(ModifierT::StaticStruct())); }
+
+	UFUNCTION(BlueprintCallable, Category = Mover)
+	CHAOSMOVER_API void CancelModifierFromHandle(FMovementModifierHandle ModifierHandle);
+
+	UFUNCTION(BlueprintPure, Category = Mover, meta = (Keywords = "HasTag"))
+	CHAOSMOVER_API bool HasGameplayTag(FGameplayTag TagToFind, bool bExactMatch) const;
+
 	struct FInitParams
 	{
 		TMap<FName, TWeakObjectPtr<UBaseMovementMode>> ModesToRegister;
@@ -73,7 +96,7 @@ public:
 	CHAOSMOVER_API void Init(const FInitParams& InitParams);
 	CHAOSMOVER_API void Deinit();
 
-	CHAOSMOVER_API void ProcessInputs(int32 PhysicsStep, const UE::ChaosMover::FSimulationInputData& InputData);
+	CHAOSMOVER_API void ProcessInputs(int32 PhysicsStep, const FMoverTimeStep& TimeStep, const UE::ChaosMover::FSimulationInputData& InputData);
 	CHAOSMOVER_API void SimulationTick(const FMoverTimeStep& TimeStep, const UE::ChaosMover::FSimulationInputData& InputData, UE::ChaosMover::FSimulationOutputData& OutputData);
 	CHAOSMOVER_API void ModifyContacts(const FMoverTimeStep& TimeStep, const UE::ChaosMover::FSimulationInputData& InputData, const UE::ChaosMover::FSimulationOutputData& OutputData, Chaos::FCollisionContactModifier& Modifier) const;
 
@@ -103,6 +126,7 @@ protected:
 	CHAOSMOVER_API virtual void OnSimulationTick(const FMoverTimeStep& TimeStep, const UE::ChaosMover::FSimulationInputData& InputData, UE::ChaosMover::FSimulationOutputData& OutputData);
 	CHAOSMOVER_API virtual void OnPostSimulationTick(const FMoverTimeStep& TimeStep, UE::ChaosMover::FSimulationOutputData& OutputData);
 	CHAOSMOVER_API virtual void OnModifyContacts(const FMoverTimeStep& TimeStep, const UE::ChaosMover::FSimulationInputData& InputData, const UE::ChaosMover::FSimulationOutputData& OutputData, Chaos::FCollisionContactModifier& Modifier) const;
+	CHAOSMOVER_API virtual void OnSimulationRollback(const FMoverTimeStep& NewTimeStep, const FMoverSyncState& PrevSyncState);
 
 	// Character-like movement sim steps
 	CHAOSMOVER_API virtual void PreSimulationTickCharacter(const IChaosCharacterMovementModeInterface& CharacterMode, const FMoverTimeStep& TimeStep, const UE::ChaosMover::FSimulationInputData& InputData);
@@ -169,7 +193,9 @@ protected:
 
 private:
 	FMoverInputCmdContext NetInputCmd;
+	FMoverSyncState NetSyncState;
 	TArray<TSharedPtr<FMoverSimulationEventData>> Events;
 
 	bool bInputCmdOverridden = false;
+	bool bSyncStateOverridden = false;
 };
diff --git a/Engine/Plugins/Experimental/ChaosMover/Source/ChaosMover/Public/ChaosMover/ChaosMoverSimulationTypes.h b/Engine/Plugins/Experimental/ChaosMover/Source/ChaosMover/Public/ChaosMover/ChaosMoverSimulationTypes.h
--- a/Engine/Plugins/Experimental/ChaosMover/Source/ChaosMover/Public/ChaosMover/ChaosMoverSimulationTypes.h
+++ b/Engine/Plugins/Experimental/ChaosMover/Source/ChaosMover/Public/ChaosMover/ChaosMoverSimulationTypes.h
@@ -3,7 +3,10 @@
 #pragma once
 
 #include "Chaos/Declares.h"
+#include "Chaos/ImplicitObject.h"
+#include "Chaos/PhysicsObject.h"
 #include "CollisionQueryParams.h"
+#include "DefaultMovementSet/MovementModifiers/StanceModifier.h"
 #include "MoverSimulationTypes.h"
 
 #include "ChaosMoverSimulationTypes.generated.h"
@@ -126,17 +129,23 @@ struct CHAOSMOVER_API FChaosMoverSimulationDefaultInputs : public FMoverDataStru
 
 	UPROPERTY(VisibleAnywhere, BlueprintReadOnly, Category=ChaosMover)
 	FVector UpDir;
+
 	UPROPERTY(VisibleAnywhere, BlueprintReadOnly, Category=ChaosMover)
 	FVector Gravity;
+
 	Chaos::FPhysicsObjectHandle PhysicsObject;
 	AActor* OwningActor;
 	UWorld* World;
+
 	UPROPERTY(VisibleAnywhere, BlueprintReadOnly, Category=ChaosMover)
 	float PhysicsObjectGravity;
+
 	UPROPERTY(VisibleAnywhere, BlueprintReadOnly, Category=ChaosMover)
 	float PawnCollisionHalfHeight;
+
 	UPROPERTY(VisibleAnywhere, BlueprintReadOnly, Category=ChaosMover)
 	float PawnCollisionRadius;
+
 	ECollisionChannel CollisionChannel;
 };
 
@@ -162,6 +171,24 @@ public:
 
 	UFUNCTION(BlueprintCallable, Category = ChaosMover)
 	virtual bool ShouldCharacterRemainUpright() const = 0;
+
+	UFUNCTION(BlueprintCallable, Category = ChaosMover)
+	virtual float GetMaxSpeed() const = 0;
+
+	UFUNCTION(BlueprintCallable, Category = ChaosMover)
+	virtual void OverrideMaxSpeed(float Value) = 0;
+
+	UFUNCTION(BlueprintCallable, Category = ChaosMover)
+	virtual void ClearMaxSpeedOverride() = 0;
+
+	UFUNCTION(BlueprintCallable, Category = ChaosMover)
+	virtual float GetAcceleration() const = 0;
+
+	UFUNCTION(BlueprintCallable, Category = ChaosMover)
+	virtual void OverrideAcceleration(float Value) = 0;
+
+	UFUNCTION(BlueprintCallable, Category = ChaosMover)
+	virtual void ClearAccelerationOverride() = 0;
 };
 
 // Interface for mover modes moving on ground like characters, using a character ground constraint
@@ -176,6 +203,7 @@ class IChaosCharacterConstraintMovementModeInterface
 	GENERATED_BODY()
 
 public:
+	virtual float GetTargetHeight() const = 0;
 	virtual bool ShouldEnableConstraint() const = 0;
 	virtual void UpdateConstraintSettings(Chaos::FCharacterGroundConstraintSettings& ConstraintSettings) const = 0;
 };
@@ -192,3 +220,25 @@ enum class EChaosMoverVelocityEffectMode : uint8
 	/** Apply as an override velocity */
 	OverrideVelocity,
 };
+
+USTRUCT()
+struct FStanceModifiedEventData : public FMoverSimulationEventData
+{
+	GENERATED_BODY()
+
+	FStanceModifiedEventData(float InEventTimeMs, EStanceMode InOldStance, EStanceMode InNewStance, FEventProcessedCallbackPtr InEventProcessedCallback = nullptr)
+		: FMoverSimulationEventData(InEventTimeMs, InEventProcessedCallback)
+		, OldStance(InOldStance)
+		, NewStance(InNewStance)
+	{
+	}
+	FStanceModifiedEventData() {}
+
+	virtual UScriptStruct* GetScriptStruct() const override
+	{
+		return FStanceModifiedEventData::StaticStruct();
+	}
+
+	EStanceMode OldStance = EStanceMode::Invalid;
+	EStanceMode NewStance = EStanceMode::Invalid;
+};
diff --git a/Engine/Plugins/Experimental/ChaosMover/Source/ChaosMover/Public/ChaosMover/Character/ChaosCharacterInputs.h b/Engine/Plugins/Experimental/ChaosMover/Source/ChaosMover/Public/ChaosMover/Character/ChaosCharacterInputs.h
--- a/Engine/Plugins/Experimental/ChaosMover/Source/ChaosMover/Public/ChaosMover/Character/ChaosCharacterInputs.h
+++ b/Engine/Plugins/Experimental/ChaosMover/Source/ChaosMover/Public/ChaosMover/Character/ChaosCharacterInputs.h
@@ -7,17 +7,18 @@
 
 #include "ChaosCharacterInputs.generated.h"
 
-USTRUCT(BlueprintType)
+/** Add this struct to the character inputs to launch the character */
+USTRUCT(MinimalAPI, BlueprintType)
 struct FChaosMoverLaunchInputs : public FMoverDataStructBase
 {
 	GENERATED_BODY()
 
 public:
 	// Launch velocity in cm/s or impulse in kg cm/s
-	UPROPERTY(BlueprintReadWrite, Category = Mover)
+	UPROPERTY(BlueprintReadWrite, Category = "Chaos Mover")
 	FVector LaunchVelocityOrImpulse = FVector::ZeroVector;
 
-	UPROPERTY(BlueprintReadWrite, Category = Mover)
+	UPROPERTY(BlueprintReadWrite, Category = "Chaos Mover")
 	EChaosMoverVelocityEffectMode Mode = EChaosMoverVelocityEffectMode::AdditiveVelocity;
 
 	virtual UScriptStruct* GetScriptStruct() const override
@@ -37,6 +38,122 @@ public:
 
 template<>
 struct TStructOpsTypeTraits< FChaosMoverLaunchInputs > : public TStructOpsTypeTraitsBase2< FChaosMoverLaunchInputs >
+{
+	enum
+	{
+		WithNetSerializer = true,
+		WithCopy = true
+	};
+};
+
+/** Add this struct to toggle crouch on the character */
+USTRUCT(MinimalAPI, BlueprintType)
+struct FChaosMoverCrouchInputs : public FMoverDataStructBase
+{
+	GENERATED_BODY()
+
+public:
+	// Launch velocity in cm/s or impulse in kg cm/s
+	UPROPERTY(BlueprintReadWrite, Category = "Chaos Mover")
+	bool bWantsToCrouch = false;
+
+	virtual UScriptStruct* GetScriptStruct() const override
+	{
+		return StaticStruct();
+	}
+	virtual FMoverDataStructBase* Clone() const override
+	{
+		return new FChaosMoverCrouchInputs(*this);
+	}
+	CHAOSMOVER_API virtual bool NetSerialize(FArchive& Ar, UPackageMap* Map, bool& bOutSuccess) override;
+	CHAOSMOVER_API virtual void ToString(FAnsiStringBuilderBase& Out) const override;
+	CHAOSMOVER_API virtual bool ShouldReconcile(const FMoverDataStructBase& AuthorityState) const override;
+	CHAOSMOVER_API virtual void Interpolate(const FMoverDataStructBase& From, const FMoverDataStructBase& To, float Pct) override;
+	CHAOSMOVER_API virtual void Merge(const FMoverDataStructBase& From) override;
+};
+
+template<>
+struct TStructOpsTypeTraits< FChaosMoverCrouchInputs > : public TStructOpsTypeTraitsBase2< FChaosMoverCrouchInputs >
+{
+	enum
+	{
+		WithNetSerializer = true,
+		WithCopy = true
+	};
+};
+
+/** Add this struct to the character inputs to override the max speed and acceleration */
+USTRUCT(MinimalAPI, BlueprintType)
+struct FChaosMovementSettingsOverrides : public FMoverDataStructBase
+{
+	GENERATED_BODY()
+
+public:
+	/** The name of the mode to apply the overrides to. If this is set to none the override will apply to the currently active mode */
+	UPROPERTY(BlueprintReadWrite, EditAnywhere, Category = "General")
+	FName ModeName = NAME_None;
+
+	/** Override maximum speed in the movement plane */
+	UPROPERTY(BlueprintReadWrite, EditAnywhere, Category = "General", meta = (ClampMin = "0", UIMin = "0", ForceUnits = "cm/s"))
+	float MaxSpeedOverride = 800.f;
+
+	/** Override max linear rate of acceleration for controlled input. May be scaled based on magnitude of input. */
+	UPROPERTY(BlueprintReadWrite, EditAnywhere, Category = "General", meta = (ClampMin = "0", UIMin = "0", ForceUnits = "cm/s^2"))
+	float AccelerationOverride = 4000.f;
+
+	virtual UScriptStruct* GetScriptStruct() const override
+	{
+		return StaticStruct();
+	}
+	virtual FMoverDataStructBase* Clone() const override
+	{
+		return new FChaosMovementSettingsOverrides(*this);
+	}
+	CHAOSMOVER_API virtual bool NetSerialize(FArchive& Ar, UPackageMap* Map, bool& bOutSuccess) override;
+	CHAOSMOVER_API virtual void ToString(FAnsiStringBuilderBase& Out) const override;
+	CHAOSMOVER_API virtual bool ShouldReconcile(const FMoverDataStructBase& AuthorityState) const override;
+	CHAOSMOVER_API virtual void Interpolate(const FMoverDataStructBase& From, const FMoverDataStructBase& To, float Pct) override;
+	CHAOSMOVER_API virtual void Merge(const FMoverDataStructBase& From) override;
+};
+
+template<>
+struct TStructOpsTypeTraits< FChaosMovementSettingsOverrides > : public TStructOpsTypeTraitsBase2< FChaosMovementSettingsOverrides >
+{
+	enum
+	{
+		WithNetSerializer = true,
+		WithCopy = true
+	};
+};
+
+/** Add this struct to the character inputs to cancel any max speed and acceleration for the input mode */
+USTRUCT(MinimalAPI, BlueprintType)
+struct FChaosMovementSettingsOverridesRemover : public FMoverDataStructBase
+{
+	GENERATED_BODY()
+
+public:
+	/** The name of the mode to apply the overrides to. If this is set to none the override will apply to the currently active mode */
+	UPROPERTY(BlueprintReadWrite, EditAnywhere, Category = "General")
+	FName ModeName = NAME_None;
+
+	virtual UScriptStruct* GetScriptStruct() const override
+	{
+		return StaticStruct();
+	}
+	virtual FMoverDataStructBase* Clone() const override
+	{
+		return new FChaosMovementSettingsOverridesRemover(*this);
+	}
+	CHAOSMOVER_API virtual bool NetSerialize(FArchive& Ar, UPackageMap* Map, bool& bOutSuccess) override;
+	CHAOSMOVER_API virtual void ToString(FAnsiStringBuilderBase& Out) const override;
+	CHAOSMOVER_API virtual bool ShouldReconcile(const FMoverDataStructBase& AuthorityState) const override;
+	CHAOSMOVER_API virtual void Interpolate(const FMoverDataStructBase& From, const FMoverDataStructBase& To, float Pct) override;
+	CHAOSMOVER_API virtual void Merge(const FMoverDataStructBase& From) override;
+};
+
+template<>
+struct TStructOpsTypeTraits< FChaosMovementSettingsOverridesRemover > : public TStructOpsTypeTraitsBase2< FChaosMovementSettingsOverridesRemover >
 {
 	enum
 	{
diff --git a/Engine/Plugins/Experimental/ChaosMover/Source/ChaosMover/Public/ChaosMover/Character/ChaosCharacterMoverComponent.h b/Engine/Plugins/Experimental/ChaosMover/Source/ChaosMover/Public/ChaosMover/Character/ChaosCharacterMoverComponent.h
--- a/Engine/Plugins/Experimental/ChaosMover/Source/ChaosMover/Public/ChaosMover/Character/ChaosCharacterMoverComponent.h
+++ b/Engine/Plugins/Experimental/ChaosMover/Source/ChaosMover/Public/ChaosMover/Character/ChaosCharacterMoverComponent.h
@@ -24,6 +24,9 @@ public:
 	CHAOSMOVER_API virtual void ProduceInput(const int32 DeltaTimeMS, FMoverInputCmdContext* Cmd) override;
 	CHAOSMOVER_API virtual void DoQueueNextMode(FName DesiredModeName, bool bShouldReenter=false) override;
 
+	CHAOSMOVER_API virtual void Crouch() override;
+	CHAOSMOVER_API virtual void UnCrouch() override;
+
 	// Broadcast when this actor lands on a valid surface.
 	UPROPERTY(BlueprintAssignable, Category = "Chaos Mover")
 	FChaosMover_OnLanded OnLanded;
@@ -33,11 +36,29 @@ public:
 	UFUNCTION(BlueprintCallable, Category = "Chaos Mover")
 	void Launch(const FVector& VelocityOrImpulse, EChaosMoverVelocityEffectMode Mode = EChaosMoverVelocityEffectMode::AdditiveVelocity);
 
+	// Override the movement mode settings on a mode
+	// If the name is not set it will apply to the current mode
+	UFUNCTION(BlueprintCallable, Category = "Chaos Mover")
+	void OverrideMovementSettings(const FChaosMovementSettingsOverrides Overrides);
+
+	// Cancel overrides of movement mode settings
+	// If the name is not set it will apply to the current mode
+	UFUNCTION(BlueprintCallable, Category = "Chaos Mover")
+	void CancelMovementSettingsOverrides(FName ModeName = NAME_None);
+
 protected:
+	CHAOSMOVER_API virtual void OnMoverPreSimulationTick(const FMoverTimeStep& TimeStep, const FMoverInputCmdContext& InputCmd) override;
 	CHAOSMOVER_API virtual void ProcessSimulationEvent(const FMoverSimulationEventData& EventData) override;
 	CHAOSMOVER_API virtual void SetAdditionalSimulationOutput(const FMoverDataCollection& Data) override;
 
 	CHAOSMOVER_API void ClearQueuedMode();
+	CHAOSMOVER_API virtual void OnStanceModified(const FStanceModifiedEventData& EventData);
+
+	FName ModeToOverrideSettings = NAME_None;
+	float MaxSpeedOverride = 0.0f;
+	float AccelerationOverride = 0.0f;
+	bool bOverrideMovementSettings = false;
+	bool bCancelMovementOverrides = false;
 
 	bool bFloorResultSet = false;
 	FFloorCheckResult LatestFloorResult;
@@ -47,4 +68,6 @@ protected:
 
 	// Queued immediate mode transition, it will be transmitted to the PT next time ProduceInput is called
 	FName QueuedModeTransitionName = NAME_None;
+
+	bool bCrouchInputPending = false;
 };
diff --git a/Engine/Plugins/Experimental/ChaosMover/Source/ChaosMover/Public/ChaosMover/Character/Modes/ChaosCharacterMovementMode.h b/Engine/Plugins/Experimental/ChaosMover/Source/ChaosMover/Public/ChaosMover/Character/Modes/ChaosCharacterMovementMode.h
--- a/Engine/Plugins/Experimental/ChaosMover/Source/ChaosMover/Public/ChaosMover/Character/Modes/ChaosCharacterMovementMode.h
+++ b/Engine/Plugins/Experimental/ChaosMover/Source/ChaosMover/Public/ChaosMover/Character/Modes/ChaosCharacterMovementMode.h
@@ -15,18 +15,18 @@ namespace Chaos
 /**
  * Base class for all Chaos character movement modes
  */
-UCLASS(Abstract, Blueprintable, BlueprintType, EditInlineNew, DefaultToInstanced)
-class CHAOSMOVER_API UChaosCharacterMovementMode :	public UChaosMovementMode,
-													public IChaosCharacterMovementModeInterface,
-													public IChaosCharacterConstraintMovementModeInterface
+UCLASS(MinimalAPI, Abstract, Blueprintable, BlueprintType, EditInlineNew, DefaultToInstanced)
+class UChaosCharacterMovementMode :	public UChaosMovementMode,
+									public IChaosCharacterMovementModeInterface,
+									public IChaosCharacterConstraintMovementModeInterface
 {
 	GENERATED_BODY()
 
 public:
-	UChaosCharacterMovementMode(const FObjectInitializer& ObjectInitializer);
+	CHAOSMOVER_API UChaosCharacterMovementMode(const FObjectInitializer& ObjectInitializer);
 
-	virtual void OnRegistered(const FName ModeName) override;
-	virtual void OnUnregistered() override;
+	CHAOSMOVER_API virtual void OnRegistered(const FName ModeName) override;
+	CHAOSMOVER_API virtual void OnUnregistered() override;
 
 	// IChaosCharacterConstraintMovementModeInterface
 	virtual bool ShouldEnableConstraint() const override
@@ -34,7 +34,7 @@ public:
 		return true;
 	}
 
-	virtual void UpdateConstraintSettings(Chaos::FCharacterGroundConstraintSettings& ConstraintSettings) const override;
+	CHAOSMOVER_API virtual void UpdateConstraintSettings(Chaos::FCharacterGroundConstraintSettings& ConstraintSettings) const override;
 
 	// IChaosCharacterMovementModeInterface
 	virtual float GetTargetHeight() const override
@@ -47,19 +47,27 @@ public:
 		return QueryRadius;
 	}
 
-	virtual float GetMaxWalkSlopeCosine() const override;
+	CHAOSMOVER_API virtual float GetMaxWalkSlopeCosine() const override;
 
 	virtual bool ShouldCharacterRemainUpright() const override
 	{
 		return bShouldCharacterRemainUpright;
 	}
 
+	CHAOSMOVER_API virtual float GetMaxSpeed() const override;
+	CHAOSMOVER_API virtual void OverrideMaxSpeed(float Value) override;
+	CHAOSMOVER_API virtual void ClearMaxSpeedOverride() override;
+
+	CHAOSMOVER_API virtual float GetAcceleration() const override;
+	CHAOSMOVER_API virtual void OverrideAcceleration(float Value) override;
+	CHAOSMOVER_API virtual void ClearAccelerationOverride() override;
+
 protected:
-	void SetTargetHeightOverride(float InTargetHeight);
-	void ClearTargetHeightOverride();
+	CHAOSMOVER_API void SetTargetHeightOverride(float InTargetHeight);
+	CHAOSMOVER_API void ClearTargetHeightOverride();
 
-	void SetQueryRadiusOverride(float InQueryRadius);
-	void ClearQueryRadiusOverride();
+	CHAOSMOVER_API void SetQueryRadiusOverride(float InQueryRadius);
+	CHAOSMOVER_API void ClearQueryRadiusOverride();
 
 	TObjectPtr<const class USharedChaosCharacterMovementSettings> SharedSettings;
 
@@ -93,4 +101,7 @@ protected:
 private:
 	float TargetHeight = 95.0f;
 	float QueryRadius = 30.0f;
+
+	TOptional<float> MaxSpeedOverride;
+	TOptional<float> AccelerationOverride;
 };
\ No newline at end of file
diff --git a/Engine/Plugins/Experimental/ChaosMover/Source/ChaosMover/Public/ChaosMover/Character/Modifiers/ChaosStanceModifier.h b/Engine/Plugins/Experimental/ChaosMover/Source/ChaosMover/Public/ChaosMover/Character/Modifiers/ChaosStanceModifier.h
new file mode 100644
--- /dev/null
+++ b/Engine/Plugins/Experimental/ChaosMover/Source/ChaosMover/Public/ChaosMover/Character/Modifiers/ChaosStanceModifier.h
@@ -0,0 +1,88 @@
+// Copyright Epic Games, Inc. All Rights Reserved.
+
+#pragma once
+
+#include "MovementModifier.h"
+#include "Chaos/ParticleHandleFwd.h"
+#include "ChaosMover/ChaosMoverSimulationTypes.h"
+#include "DefaultMovementSet/MovementModifiers/StanceModifier.h"
+
+#include "ChaosStanceModifier.generated.h"
+
+
+USTRUCT(MinimalAPI, BlueprintType)
+struct FChaosStanceModifier : public FMovementModifierBase
+{
+	GENERATED_BODY()
+
+public:
+	CHAOSMOVER_API FChaosStanceModifier();
+	virtual ~FChaosStanceModifier() override {};
+
+	EStanceMode ActiveStance;
+
+	/** Whether to cancel the modifier when the movement mode changes */
+	UPROPERTY(BlueprintReadWrite, EditAnywhere, Category = Parameters)
+	bool bCancelOnModeChange = false;
+
+	/** Height of the modified capsule */
+	UPROPERTY(BlueprintReadWrite, EditAnywhere, Category = Parameters, meta = (ClampMin = "0", UIMin = "0", ForceUnits = "cm"))
+	float ModifiedCapsuleHalfHeight;
+
+	/** Radius of the modified capsule */
+	UPROPERTY(BlueprintReadWrite, EditAnywhere, Category = Parameters, meta = (ClampMin = "0", UIMin = "0", ForceUnits = "cm"))
+	float ModifiedCapsuleRadius;
+
+	/** Ground clearance of the modified capsule */
+	UPROPERTY(BlueprintReadWrite, EditAnywhere, Category = Parameters, meta = (ClampMin = "0", UIMin = "0", ForceUnits = "cm"))
+	float ModifiedCapsuleGroundClearance;
+
+	/** Height of the default unmodified capsule */
+	UPROPERTY(BlueprintReadWrite, EditAnywhere, Category = Parameters, meta = (ClampMin = "0", UIMin = "0", ForceUnits = "cm"))
+	float DefaultCapsuleHalfHeight;
+
+	/** Radius of the default unmodified capsule */
+	UPROPERTY(BlueprintReadWrite, EditAnywhere, Category = Parameters, meta = (ClampMin = "0", UIMin = "0", ForceUnits = "cm"))
+	float DefaultCapsuleRadius;
+
+	/** Ground clearance of the default unmodified capsule */
+	UPROPERTY(BlueprintReadWrite, EditAnywhere, Category = Parameters, meta = (ClampMin = "0", UIMin = "0", ForceUnits = "cm"))
+	float DefaultCapsuleGroundClearance;
+
+	/** Override Max Speed */
+	UPROPERTY(BlueprintReadWrite, EditAnywhere, Category = Parameters, meta = (ClampMin = "0", UIMin = "0", ForceUnits = "cm/s"))
+	TOptional<float> MaxSpeedOverride;
+
+	/** Override Acceleration */
+	UPROPERTY(BlueprintReadWrite, EditAnywhere, Category = Parameters, meta = (ClampMin = "0", UIMin = "0", ForceUnits = "cm/s^2"))
+	TOptional<float> AccelerationOverride;
+
+	CHAOSMOVER_API virtual bool HasGameplayTag(FGameplayTag TagToFind, bool bExactMatch) const override;
+
+	CHAOSMOVER_API virtual void OnStart_Async(const FMovementModifierParams_Async& Params) override;
+	CHAOSMOVER_API virtual void OnEnd_Async(const FMovementModifierParams_Async& Params) override;
+	CHAOSMOVER_API virtual void OnPostMovement_Async(const FMovementModifierParams_Async& Params) override;
+
+	CHAOSMOVER_API virtual FMovementModifierBase* Clone() const override;
+	CHAOSMOVER_API virtual void NetSerialize(FArchive& Ar) override;
+	CHAOSMOVER_API virtual UScriptStruct* GetScriptStruct() const override;
+	CHAOSMOVER_API virtual FString ToSimpleString() const override;
+	CHAOSMOVER_API virtual void AddReferencedObjects(class FReferenceCollector& Collector) override;
+
+protected:
+	CHAOSMOVER_API virtual void UpdateStance(const FMovementModifierParams_Async& AsyncParams, EStanceMode NewStance);
+
+	CHAOSMOVER_API void UpdateCapsule(Chaos::FPBDRigidParticleHandle* ParticleHandle, float NewRadius, float NewHalfHeight, float NewGroundClearance, float TargetHeight);
+	CHAOSMOVER_API const Chaos::FCapsule* GetCapsule(Chaos::FPBDRigidParticleHandle* ParticleHandle) const;
+
+	FName CurrentModeName = NAME_None;
+};
+
+template<>
+struct TStructOpsTypeTraits< FChaosStanceModifier > : public TStructOpsTypeTraitsBase2< FChaosStanceModifier >
+{
+	enum
+	{
+		WithCopy = true
+	};
+};
\ No newline at end of file
diff --git a/Engine/Plugins/Experimental/ChaosMover/Source/ChaosMover/Public/ChaosMover/Character/Transitions/ChaosCharacterCrouchCheck.h b/Engine/Plugins/Experimental/ChaosMover/Source/ChaosMover/Public/ChaosMover/Character/Transitions/ChaosCharacterCrouchCheck.h
new file mode 100644
--- /dev/null
+++ b/Engine/Plugins/Experimental/ChaosMover/Source/ChaosMover/Public/ChaosMover/Character/Transitions/ChaosCharacterCrouchCheck.h
@@ -0,0 +1,71 @@
+// Copyright Epic Games, Inc. All Rights Reserved.
+
+#pragma once
+
+#include "ChaosMover/ChaosMovementModeTransition.h"
+#include "MovementModifier.h"
+
+#include "ChaosCharacterCrouchCheck.generated.h"
+
+
+UCLASS(MinimalAPI, Blueprintable, BlueprintType, EditInlineNew, DefaultToInstanced)
+class UChaosCharacterCrouchCheck : public UChaosMovementModeTransition
+{
+	GENERATED_BODY()
+
+public:
+	CHAOSMOVER_API UChaosCharacterCrouchCheck(const FObjectInitializer& ObjectInitializer);
+
+	CHAOSMOVER_API virtual void OnRegistered();
+	CHAOSMOVER_API virtual void OnUnregistered();
+
+	CHAOSMOVER_API virtual FTransitionEvalResult Evaluate_Implementation(const FSimulationTickParams& Params) const override;
+	CHAOSMOVER_API virtual void Trigger_Implementation(const FSimulationTickParams& Params) override;
+
+	/** Whether to cancel the modifier when the movement mode changes */
+	UPROPERTY(BlueprintReadWrite, EditAnywhere, Category = Parameters)
+	bool bCancelOnModeChange = false;
+
+	/** Height of the modified capsule */
+	UPROPERTY(BlueprintReadWrite, EditAnywhere, Category = Parameters, meta = (ClampMin = "0", UIMin = "0", ForceUnits = "cm"))
+	TOptional<float> CapsuleHalfHeight;
+
+	/** Radius of the modified capsule */
+	UPROPERTY(BlueprintReadWrite, EditAnywhere, Category = Parameters, meta = (ClampMin = "0", UIMin = "0", ForceUnits = "cm"))
+	TOptional<float> CapsuleRadius;
+
+	/** Ground clearance of the modified capsule */
+	UPROPERTY(BlueprintReadWrite, EditAnywhere, Category = Parameters, meta = (ClampMin = "0", UIMin = "0", ForceUnits = "cm"))
+	TOptional<float> CapsuleGroundClearance;
+
+	/** Override Max Speed */
+	UPROPERTY(BlueprintReadWrite, EditAnywhere, Category = Parameters, meta = (ClampMin = "0", UIMin = "0", ForceUnits = "cm/s"))
+	TOptional<float> MaxSpeed;
+
+	/** Override Acceleration */
+	UPROPERTY(BlueprintReadWrite, EditAnywhere, Category = Parameters, meta = (ClampMin = "0", UIMin = "0", ForceUnits = "cm/s^2"))
+	TOptional<float> Acceleration;
+
+	/** Name of movement mode to transition to when crouched */
+	UPROPERTY(BlueprintReadWrite, EditAnywhere, Category = Parameters)
+	TOptional<FName> TransitionToCrouchedMode;
+
+	/** Name of movement mode to transition to when uncrouched */
+	UPROPERTY(BlueprintReadWrite, EditAnywhere, Category = Parameters)
+	TOptional<FName> TransitionToUncrouchedMode;
+
+protected:
+	CHAOSMOVER_API virtual bool CanCrouch(const FSimulationTickParams& Params) const;
+	CHAOSMOVER_API virtual bool CanUncrouch(const FSimulationTickParams& Params) const;
+
+	CHAOSMOVER_API virtual void Crouch(const FSimulationTickParams& Params);
+	CHAOSMOVER_API virtual void Uncrouch(const FSimulationTickParams& Params);
+
+	FMovementModifierHandle CrouchModifierHandle;
+
+	float OriginalCapsuleHalfHeight = 0.0f;
+	float OriginalCapsuleRadius = 0.0f;
+
+	mutable bool bTriggerCrouch = false;
+	mutable bool bTriggerUncrouch = false;
+};
\ No newline at end of file
diff --git a/Engine/Plugins/Experimental/ChaosMover/Source/ChaosMover/Public/ChaosMover/Utilities/ChaosMoverQueryUtils.h b/Engine/Plugins/Experimental/ChaosMover/Source/ChaosMover/Public/ChaosMover/Utilities/ChaosMoverQueryUtils.h
--- a/Engine/Plugins/Experimental/ChaosMover/Source/ChaosMover/Public/ChaosMover/Utilities/ChaosMoverQueryUtils.h
+++ b/Engine/Plugins/Experimental/ChaosMover/Source/ChaosMover/Public/ChaosMover/Utilities/ChaosMoverQueryUtils.h
@@ -27,4 +27,19 @@ namespace UE::ChaosMover::Utils
 
 	extern CHAOSMOVER_API void FloorSweep_Internal(const FFloorSweepParams& Params, FFloorCheckResult& OutFloorResult, FWaterCheckResult& OutWaterResult);
 
+	struct FCapsuleOverlapTestParams
+	{
+		FCollisionResponseParams ResponseParams;
+		FCollisionQueryParams QueryParams;
+		FVector Location;
+		FVector UpDir;
+		const UWorld* World;
+		float CapsuleHalfHeight;
+		float CapsuleRadius;
+		ECollisionChannel CollisionChannel;
+		bool bBlockOnDynamic = false;
+	};
+
+	extern CHAOSMOVER_API bool CapsuleOverlapTest_Internal(const FCapsuleOverlapTestParams& Params);
+
 }
\ No newline at end of file
