Chaos Mover - Move ChaosMoverStateMachine header to public folder

#rnx

diff --git a/Engine/Plugins/Experimental/ChaosMover/Source/ChaosMover/Private/ChaosMoverStateMachine.cpp b/Engine/Plugins/Experimental/ChaosMover/Source/ChaosMover/Private/ChaosMoverStateMachine.cpp
--- a/Engine/Plugins/Experimental/ChaosMover/Source/ChaosMover/Private/ChaosMoverStateMachine.cpp
+++ b/Engine/Plugins/Experimental/ChaosMover/Source/ChaosMover/Private/ChaosMoverStateMachine.cpp
@@ -1,7 +1,7 @@
 // Copyright Epic Games, Inc. All Rights Reserved.
 
 
-#include "ChaosMoverStateMachine.h"
+#include "ChaosMover/ChaosMoverStateMachine.h"
 
 #include "Backends/ChaosMoverSubsystem.h"
 #include "ChaosMover/ChaosMoverDeveloperSettings.h"
diff --git a/Engine/Plugins/Experimental/ChaosMover/Source/ChaosMover/Private/ChaosMoverStateMachine.h b/Engine/Plugins/Experimental/ChaosMover/Source/ChaosMover/Private/ChaosMoverStateMachine.h
deleted file mode 100644
diff --git a/Engine/Plugins/Experimental/ChaosMover/Source/ChaosMover/Public/ChaosMover/ChaosMoverSimulation.h b/Engine/Plugins/Experimental/ChaosMover/Source/ChaosMover/Public/ChaosMover/ChaosMoverSimulation.h
--- a/Engine/Plugins/Experimental/ChaosMover/Source/ChaosMover/Public/ChaosMover/ChaosMoverSimulation.h
+++ b/Engine/Plugins/Experimental/ChaosMover/Source/ChaosMover/Public/ChaosMover/ChaosMoverSimulation.h
@@ -5,7 +5,7 @@
 #include "Chaos/Declares.h"
 #include "ChaosMover/ChaosMoverSimulationTypes.h"
 #include "ChaosMover/PathedMovement/ChaosPathedMovementTypes.h"
-#include "ChaosMoverStateMachine.h"
+#include "ChaosMover/ChaosMoverStateMachine.h"
 #include "MoveLibrary/MovementMixer.h"
 #include "MoveLibrary/MoverBlackboard.h"
 #include "MoverTypes.h"
diff --git a/Engine/Plugins/Experimental/ChaosMover/Source/ChaosMover/Public/ChaosMover/ChaosMoverStateMachine.h b/Engine/Plugins/Experimental/ChaosMover/Source/ChaosMover/Public/ChaosMover/ChaosMoverStateMachine.h
new file mode 100644
--- /dev/null
+++ b/Engine/Plugins/Experimental/ChaosMover/Source/ChaosMover/Public/ChaosMover/ChaosMoverStateMachine.h
@@ -0,0 +1,115 @@
+// Copyright Epic Games, Inc. All Rights Reserved.
+
+#pragma once
+
+#include "CoreMinimal.h"
+#include "Engine/EngineTypes.h"
+#include "InstantMovementEffect.h"
+#include "MovementMode.h"
+#include "Templates/SubclassOf.h"
+#include "ChaosMover/ChaosMoverSimulationTypes.h"
+
+class UChaosMoverSimulation;
+
+namespace UE::ChaosMover
+{
+	class FMoverStateMachine
+	{
+	public:
+		CHAOSMOVER_API FMoverStateMachine();
+		CHAOSMOVER_API virtual ~FMoverStateMachine();
+
+		struct FInitParams
+		{
+			TWeakObjectPtr<UChaosMoverSimulation> Simulation;
+			TWeakObjectPtr<UNullMovementMode> NullMovementMode;
+			TWeakObjectPtr<UImmediateMovementModeTransition> ImmediateMovementModeTransition;
+		};
+		CHAOSMOVER_API void Init(const FInitParams& Params);
+
+		CHAOSMOVER_API void RegisterMovementMode(FName ModeName, TWeakObjectPtr<UBaseMovementMode> Mode, bool bIsDefaultMode = false);
+
+		CHAOSMOVER_API void UnregisterMovementMode(FName ModeName);
+		CHAOSMOVER_API void ClearAllMovementModes();
+
+		CHAOSMOVER_API void RegisterGlobalTransition(TWeakObjectPtr<UBaseMovementModeTransition> Transition);
+		CHAOSMOVER_API void UnregisterGlobalTransition(TWeakObjectPtr<UBaseMovementModeTransition> Transition);
+		CHAOSMOVER_API void ClearAllGlobalTransitions();
+
+		CHAOSMOVER_API FName GetDefaultModeName() const;
+		CHAOSMOVER_API void SetDefaultMode(FName NewDefaultModeName);
+
+		CHAOSMOVER_API void QueueNextMode(FName DesiredNextModeName, bool bShouldReenter = false);
+		CHAOSMOVER_API void SetModeImmediately(FName DesiredModeName, bool bShouldReenter = false);
+		CHAOSMOVER_API void ClearQueuedMode();
+
+		CHAOSMOVER_API void OnSimulationTick(const FMoverTimeStep& TimeStep, const FMoverTickStartData& StartState, UMoverBlackboard* SimBlackboard, UMovementMixer* MovementMixer, FMoverTickEndData& OutputState);
+		CHAOSMOVER_API void OnSimulationRollback(const FMoverTimeStep& NewTimeStep, const FMoverSyncState& InvalidSyncState, const FMoverSyncState& NewSyncState);
+
+		FName GetCurrentModeName() const
+		{
+			return CurrentModeName;
+		}
+
+		CHAOSMOVER_API const TWeakObjectPtr<UBaseMovementMode> GetCurrentMode() const;
+		CHAOSMOVER_API const TWeakObjectPtr<UBaseMovementMode> FindMovementMode(FName ModeName) const;
+		CHAOSMOVER_API TWeakObjectPtr<UBaseMovementMode> FindMovementMode_Mutable(FName ModeName);
+
+		CHAOSMOVER_API void QueueLayeredMove(TSharedPtr<FLayeredMoveBase> Move);
+
+		CHAOSMOVER_API void QueueInstantMovementEffect(TSharedPtr<FInstantMovementEffect> Effect);
+
+		CHAOSMOVER_API FMovementModifierHandle QueueMovementModifier(TSharedPtr<FMovementModifierBase> Modifier);
+		CHAOSMOVER_API void CancelModifierFromHandle(FMovementModifierHandle ModifierHandle);
+		CHAOSMOVER_API const FMovementModifierBase* FindQueuedModifier(FMovementModifierHandle ModifierHandle) const;
+		CHAOSMOVER_API const FMovementModifierBase* FindQueuedModifierByType(const UScriptStruct* ModifierType) const;
+
+		void SetOwnerActorName(const FString& InOwnerActorName);
+		void SetOwnerActorLocalNetRole(ENetRole InOwnerActorLocalNetRole);
+
+	protected:
+
+		TMap<FName, TWeakObjectPtr<UBaseMovementMode>> Modes;
+		TArray<TWeakObjectPtr<UBaseMovementModeTransition>> GlobalTransitions;
+
+		TWeakObjectPtr<UImmediateMovementModeTransition> QueuedModeTransitionWeakPtr;
+
+		TWeakObjectPtr<UImmediateMovementModeTransition> ImmediateMovementModeTransitionWeakPtr;
+		TWeakObjectPtr<UNullMovementMode> NullMovementModeWeakPtr;
+
+		FString OwnerActorName;
+		ENetRole OwnerActorLocalNetRole;
+
+		FName DefaultModeName = NAME_None;
+		FName CurrentModeName = NAME_None;
+
+		/** Moves that are queued to be added to the simulation at the start of the next sim subtick */
+		TArray<TSharedPtr<FLayeredMoveBase>> QueuedLayeredMoves;
+
+		/** Effects that are queued to be applied to the simulation at the start of the next sim subtick or at the end of this tick */
+		TArray<TSharedPtr<FInstantMovementEffect>> QueuedInstantEffects;
+
+		/** Modifiers that are queued to be added to the simulation at the start of the next sim subtick. */
+		TArray<TSharedPtr<FMovementModifierBase>> QueuedMovementModifiers;
+
+		/** Modifiers that are to be canceled at the start of the next sim subtick. */
+		TArray<FMovementModifierHandle> ModifiersToCancel;
+
+		// Internal-use-only tick data structs, for efficiency since they typically have the same contents from frame to frame
+		FMoverTickStartData WorkingSubstepStartData;
+		FSimulationTickParams WorkingSimTickParams;
+
+		TWeakObjectPtr<UChaosMoverSimulation> Simulation;
+
+	private:
+		void ConstructDefaultModes();
+		void AdvanceToNextMode();
+		void FlushQueuedMovesToGroup(FLayeredMoveGroup& Group);
+		void FlushQueuedModifiersToGroup(FMovementModifierGroup& ModifierGroup);
+		void FlushModifierCancellationsToGroup(FMovementModifierGroup& ActiveModifierGroup);
+		bool HasAnyInstantEffectsQueued() const;
+		bool ApplyInstantEffects(const FMoverTickStartData& SubstepStartData, const FMoverTimeStep& SubTimeStep, FMoverSyncState& OutputState);
+	
+		float InternalSimTimeMs = 0.0f;
+	};
+}
\ No newline at end of file
