Chaos Mover - Make falling mode data public

#rnx

diff --git a/Engine/Plugins/Experimental/ChaosMover/Source/ChaosMover/Public/ChaosMover/Character/Modes/ChaosFallingMode.h b/Engine/Plugins/Experimental/ChaosMover/Source/ChaosMover/Public/ChaosMover/Character/Modes/ChaosFallingMode.h
--- a/Engine/Plugins/Experimental/ChaosMover/Source/ChaosMover/Public/ChaosMover/Character/Modes/ChaosFallingMode.h
+++ b/Engine/Plugins/Experimental/ChaosMover/Source/ChaosMover/Public/ChaosMover/Character/Modes/ChaosFallingMode.h
@@ -20,7 +20,6 @@ public:
 	virtual void GenerateMove_Implementation(const FMoverTickStartData& StartState, const FMoverTimeStep& TimeStep, FProposedMove& OutProposedMove) const override;
 	virtual void SimulationTick_Implementation(const FSimulationTickParams& Params, FMoverTickEndData& OutputState) override;
 
-protected:
 	/**
 	 * If true, actor will land and lose all speed in the vertical direction upon landing. If false, actor's vertical speed will be redirected based on the surface normal it hit.
 	 * Note: Actor's horizontal speed will not be affected if true. If false, horizontal speed may be increased on landing.
