Chaos Mover - Fix incorrect dll export in ChaosMoverSimulation

#rnx

diff --git a/Engine/Plugins/Experimental/ChaosMover/Source/ChaosMover/Public/ChaosMover/ChaosMoverSimulation.h b/Engine/Plugins/Experimental/ChaosMover/Source/ChaosMover/Public/ChaosMover/ChaosMoverSimulation.h
--- a/Engine/Plugins/Experimental/ChaosMover/Source/ChaosMover/Public/ChaosMover/ChaosMoverSimulation.h
+++ b/Engine/Plugins/Experimental/ChaosMover/Source/ChaosMover/Public/ChaosMover/ChaosMoverSimulation.h
@@ -67,7 +67,7 @@ public:
 
 	/** Find a movement modifier of a specific type in this components movement modifiers. If not found, null will be returned. */
 	template <typename ModifierT = FMovementModifierBase UE_REQUIRES(std::is_base_of_v<FMovementModifierBase, ModifierT>)>
-	CHAOSMOVER_API const ModifierT* FindMovementModifierByType() const { return static_cast<const ModifierT*>(FindMovementModifierByType(ModifierT::StaticStruct())); }
+	const ModifierT* FindMovementModifierByType() const { return static_cast<const ModifierT*>(FindMovementModifierByType(ModifierT::StaticStruct())); }
 
 	UFUNCTION(BlueprintCallable, Category = Mover)
 	CHAOSMOVER_API void CancelModifierFromHandle(FMovementModifierHandle ModifierHandle);
