Chaos Mover - Fix uninitialized variables in ChaosStanceModifier and FChaosPathedMovementModeDebugData

#rnx

diff --git a/Engine/Plugins/Experimental/ChaosMover/Source/ChaosMover/Public/ChaosMover/Character/Modifiers/ChaosStanceModifier.h b/Engine/Plugins/Experimental/ChaosMover/Source/ChaosMover/Public/ChaosMover/Character/Modifiers/ChaosStanceModifier.h
--- a/Engine/Plugins/Experimental/ChaosMover/Source/ChaosMover/Public/ChaosMover/Character/Modifiers/ChaosStanceModifier.h
+++ b/Engine/Plugins/Experimental/ChaosMover/Source/ChaosMover/Public/ChaosMover/Character/Modifiers/ChaosStanceModifier.h
@@ -27,27 +27,27 @@ public:
 
 	/** Height of the modified capsule */
 	UPROPERTY(BlueprintReadWrite, EditAnywhere, Category = Parameters, meta = (ClampMin = "0", UIMin = "0", ForceUnits = "cm"))
-	float ModifiedCapsuleHalfHeight;
+	float ModifiedCapsuleHalfHeight = 60.0f;
 
 	/** Radius of the modified capsule */
 	UPROPERTY(BlueprintReadWrite, EditAnywhere, Category = Parameters, meta = (ClampMin = "0", UIMin = "0", ForceUnits = "cm"))
-	float ModifiedCapsuleRadius;
+	float ModifiedCapsuleRadius = 30.0f;
 
 	/** Ground clearance of the modified capsule */
 	UPROPERTY(BlueprintReadWrite, EditAnywhere, Category = Parameters, meta = (ClampMin = "0", UIMin = "0", ForceUnits = "cm"))
-	float ModifiedCapsuleGroundClearance;
+	float ModifiedCapsuleGroundClearance = 10.0f;
 
 	/** Height of the default unmodified capsule */
 	UPROPERTY(BlueprintReadWrite, EditAnywhere, Category = Parameters, meta = (ClampMin = "0", UIMin = "0", ForceUnits = "cm"))
-	float DefaultCapsuleHalfHeight;
+	float DefaultCapsuleHalfHeight = 80.0f;
 
 	/** Radius of the default unmodified capsule */
 	UPROPERTY(BlueprintReadWrite, EditAnywhere, Category = Parameters, meta = (ClampMin = "0", UIMin = "0", ForceUnits = "cm"))
-	float DefaultCapsuleRadius;
+	float DefaultCapsuleRadius = 30.0f;
 
 	/** Ground clearance of the default unmodified capsule */
 	UPROPERTY(BlueprintReadWrite, EditAnywhere, Category = Parameters, meta = (ClampMin = "0", UIMin = "0", ForceUnits = "cm"))
-	float DefaultCapsuleGroundClearance;
+	float DefaultCapsuleGroundClearance = 10.0f;
 
 	/** Override Max Speed */
 	UPROPERTY(BlueprintReadWrite, EditAnywhere, Category = Parameters, meta = (ClampMin = "0", UIMin = "0", ForceUnits = "cm/s"))
diff --git a/Engine/Plugins/Experimental/ChaosMover/Source/ChaosMover/Public/ChaosMover/PathedMovement/ChaosPathedMovementTypes.h b/Engine/Plugins/Experimental/ChaosMover/Source/ChaosMover/Public/ChaosMover/PathedMovement/ChaosPathedMovementTypes.h
--- a/Engine/Plugins/Experimental/ChaosMover/Source/ChaosMover/Public/ChaosMover/PathedMovement/ChaosPathedMovementTypes.h
+++ b/Engine/Plugins/Experimental/ChaosMover/Source/ChaosMover/Public/ChaosMover/PathedMovement/ChaosPathedMovementTypes.h
@@ -165,28 +165,28 @@ struct FChaosPathedMovementModeDebugData : public FMoverDataStructBase
 	FChaosPathedMovementState PostSimMoveState;
 
 	UPROPERTY(VisibleAnywhere, Category = "Mover Info")
-	FVector PreviousRelativeTranslation;
+	FVector PreviousRelativeTranslation = FVector::ZeroVector;
 
 	UPROPERTY(VisibleAnywhere, Category = "Mover Info")
-	FRotator PreviousRelativeRotation;
+	FRotator PreviousRelativeRotation = FRotator::ZeroRotator;
 
 	UPROPERTY(VisibleAnywhere, Category = "Mover Info")
-	FVector PreviousWorldTranslation;
+	FVector PreviousWorldTranslation = FVector::ZeroVector;
 
 	UPROPERTY(VisibleAnywhere, Category = "Mover Info")
-	FRotator PreviousWorldRotation;
+	FRotator PreviousWorldRotation = FRotator::ZeroRotator;
 
 	UPROPERTY(VisibleAnywhere, Category = "Mover Info")
-	FVector NewRelativeTranslation;
+	FVector NewRelativeTranslation = FVector::ZeroVector;
 
 	UPROPERTY(VisibleAnywhere, Category = "Mover Info")
-	FRotator NewRelativeRotation;
+	FRotator NewRelativeRotation = FRotator::ZeroRotator;
 
 	UPROPERTY(VisibleAnywhere, Category = "Mover Info")
-	FVector NewWorldTranslation;
+	FVector NewWorldTranslation = FVector::ZeroVector;
 
 	UPROPERTY(VisibleAnywhere, Category = "Mover Info")
-	FRotator NewWorldRotation;
+	FRotator NewWorldRotation = FRotator::ZeroRotator;
 
 	UPROPERTY(VisibleAnywhere, Category = "Mover Info")
 	bool bPlaybackStartedThisFrame = false;
