Chaos Mover - Add missing export symbols in ChaosMoverSimulation

#rnx

diff --git a/Engine/Plugins/Experimental/ChaosMover/Source/ChaosMover/Public/ChaosMover/ChaosMoverSimulation.h b/Engine/Plugins/Experimental/ChaosMover/Source/ChaosMover/Public/ChaosMover/ChaosMoverSimulation.h
--- a/Engine/Plugins/Experimental/ChaosMover/Source/ChaosMover/Public/ChaosMover/ChaosMoverSimulation.h
+++ b/Engine/Plugins/Experimental/ChaosMover/Source/ChaosMover/Public/ChaosMover/ChaosMoverSimulation.h
@@ -109,9 +109,9 @@ public:
 
 	// Pathed Movement Simulation API
 	// Returns the reference frame of the pathed movement
-	virtual const FTransform& GetPathedMovementBasis() const;
+	CHAOSMOVER_API virtual const FTransform& GetPathedMovementBasis() const;
 	// Sets the reference frame of the pathed movement
-	virtual void SetPathedMovementBasis(const FTransform& InPathedMovementBasis);
+	CHAOSMOVER_API virtual void SetPathedMovementBasis(const FTransform& InPathedMovementBasis);
 
 	//~ Debugging Util functions
 	// Collection for holding extra debug data, that will be sent to the Chaos Visual Debugger for debugging
@@ -155,8 +155,8 @@ protected:
 	// General purpose actuation joint constraint
 	CHAOSMOVER_API void EnableActuationConstraint();
 	CHAOSMOVER_API void DisableActuationConstraint();
-	void SetActuationTargetTransform(const FTransform& TargetTransform);
-	void TeleportActuationTarget(const FTransform& TargetTransform, bool AlsoTeleportControlledParticle = false);
+	CHAOSMOVER_API void SetActuationTargetTransform(const FTransform& TargetTransform);
+	CHAOSMOVER_API void TeleportActuationTarget(const FTransform& TargetTransform, bool AlsoTeleportControlledParticle = false);
 	
 	CHAOSMOVER_API void TraceMoverData(const UE::ChaosMover::FSimulationOutputData& OutputData);
 
