Chaos Mover - Make it possible to override the simulation class on a backend

#rnx

diff --git a/Engine/Plugins/Experimental/ChaosMover/Source/ChaosMover/Private/Backends/ChaosMoverBackend.cpp b/Engine/Plugins/Experimental/ChaosMover/Source/ChaosMover/Private/Backends/ChaosMoverBackend.cpp
--- a/Engine/Plugins/Experimental/ChaosMover/Source/ChaosMover/Private/Backends/ChaosMoverBackend.cpp
+++ b/Engine/Plugins/Experimental/ChaosMover/Source/ChaosMover/Private/Backends/ChaosMoverBackend.cpp
@@ -31,6 +31,8 @@ UChaosMoverBackendComponent::UChaosMoverBackendComponent()
 	bWantsInitializeComponent = true;
 	bAutoActivate = true;
 
+	SimulationClass = UChaosMoverSimulation::StaticClass();
+
 	if (const Chaos::FPhysicsSolver* Solver = GetPhysicsSolver())
 	{
 		bIsUsingAsyncPhysics = Solver->IsUsingAsyncResults();
@@ -57,8 +59,6 @@ UChaosMoverBackendComponent::UChaosMoverBackendComponent()
 			}
 		}
 	}
-
-	Simulation = CreateDefaultSubobject<UChaosMoverSimulation>(TEXT("ChaosMoverSimulation"));
 }
 
 void UChaosMoverBackendComponent::InitializeComponent()
@@ -67,8 +67,10 @@ void UChaosMoverBackendComponent::InitializeComponent()
 
 	if (UWorld* World = GetWorld(); World && World->IsGameWorld())
 	{
-		NullMovementMode = NewObject<UNullMovementMode>(&GetMoverComponent(), TEXT("NullMovementMode"));
-		ImmediateModeTransition = NewObject<UImmediateMovementModeTransition>(&GetMoverComponent(), TEXT("ImmediateModeTransition"));
+		Simulation = NewObject<UChaosMoverSimulation>(GetOwner(), SimulationClass, TEXT("ChaosMoverSimulation"), RF_Transient);
+
+		NullMovementMode = NewObject<UNullMovementMode>(&GetMoverComponent(), TEXT("NullMovementMode"), RF_Transient);
+		ImmediateModeTransition = NewObject<UImmediateMovementModeTransition>(&GetMoverComponent(), TEXT("ImmediateModeTransition"), RF_Transient);
 		TransformOnInit = GetMoverComponent().GetUpdatedComponentTransform();
 
 		// Create NetworkPhysicsComponent
@@ -84,7 +86,7 @@ void UChaosMoverBackendComponent::InitializeComponent()
 			}
 			else
 			{
-				NetworkPhysicsComponent = NewObject<UNetworkPhysicsComponent>(GetOwner(), TEXT("PhysMover_NetworkPhysicsComponent"));
+				NetworkPhysicsComponent = NewObject<UNetworkPhysicsComponent>(GetOwner(), TEXT("PhysMover_NetworkPhysicsComponent"), RF_Transient);
 
 				// This isn't technically a DSO component, but set it net addressable as though it is
 				NetworkPhysicsComponent->SetNetAddressable();
diff --git a/Engine/Plugins/Experimental/ChaosMover/Source/ChaosMover/Public/ChaosMover/Backends/ChaosMoverBackend.h b/Engine/Plugins/Experimental/ChaosMover/Source/ChaosMover/Public/ChaosMover/Backends/ChaosMoverBackend.h
--- a/Engine/Plugins/Experimental/ChaosMover/Source/ChaosMover/Public/ChaosMover/Backends/ChaosMoverBackend.h
+++ b/Engine/Plugins/Experimental/ChaosMover/Source/ChaosMover/Public/ChaosMover/Backends/ChaosMoverBackend.h
@@ -75,6 +75,8 @@ protected:
 	UFUNCTION()
 	CHAOSMOVER_API void HandleOwningPawnControllerChanged_Server(APawn* OwnerPawn, AController* OldController, AController* NewController);
 
+	TSubclassOf<UChaosMoverSimulation> SimulationClass;
+
 	UPROPERTY(Transient)
 	TObjectPtr<UNetworkPhysicsComponent> NetworkPhysicsComponent;
 
