Mover: upgrading absolute timesteps to double-precision

diff --git a/Engine/Plugins/Experimental/ChaosMover/Source/ChaosMover/Private/Backends/ChaosMoverAsyncCallback.cpp b/Engine/Plugins/Experimental/ChaosMover/Source/ChaosMover/Private/Backends/ChaosMoverAsyncCallback.cpp
--- a/Engine/Plugins/Experimental/ChaosMover/Source/ChaosMover/Private/Backends/ChaosMoverAsyncCallback.cpp
+++ b/Engine/Plugins/Experimental/ChaosMover/Source/ChaosMover/Private/Backends/ChaosMoverAsyncCallback.cpp
@@ -87,7 +87,7 @@ namespace UE::ChaosMover
 
 	void FAsyncCallback::GeCurrentMoverTimeStep(const FAsyncCallbackInput* AsyncInput, FMoverTimeStep& OutMoverTimeStep) const
 	{
-		OutMoverTimeStep.BaseSimTimeMs = GetSimTime_Internal() * 1000.0f;
+		OutMoverTimeStep.BaseSimTimeMs = GetSimTime_Internal() * 1000.0;
 		OutMoverTimeStep.StepMs = GetDeltaTime_Internal() * 1000.0f;
 		if (AsyncInput && AsyncInput->PhysicsSolver)
 		{
diff --git a/Engine/Plugins/Experimental/ChaosMover/Source/ChaosMover/Private/Backends/ChaosMoverBackend.cpp b/Engine/Plugins/Experimental/ChaosMover/Source/ChaosMover/Private/Backends/ChaosMoverBackend.cpp
--- a/Engine/Plugins/Experimental/ChaosMover/Source/ChaosMover/Private/Backends/ChaosMoverBackend.cpp
+++ b/Engine/Plugins/Experimental/ChaosMover/Source/ChaosMover/Private/Backends/ChaosMoverBackend.cpp
@@ -317,7 +317,7 @@ void UChaosMoverBackendComponent::InitSimulation()
 	FMoverTimeStep TimeStep;
 	if (Chaos::FPhysicsSolver* Solver = GetPhysicsSolver())
 	{
-		TimeStep.BaseSimTimeMs = Solver->GetPhysicsResultsTime_External() * 1000.0f;
+		TimeStep.BaseSimTimeMs = Solver->GetPhysicsResultsTime_External() * 1000.0;
 		TimeStep.ServerFrame = Solver->GetCurrentFrame();
 		TimeStep.StepMs = Solver->GetAsyncDeltaTime() * 1000.0f;
 	}
@@ -375,9 +375,9 @@ void UChaosMoverBackendComponent::EndPlay(const EEndPlayReason::Type EndPlayReas
 	Super::EndPlay(EndPlayReason);
 }
 
-float UChaosMoverBackendComponent::GetCurrentSimTimeMs()
+double UChaosMoverBackendComponent::GetCurrentSimTimeMs()
 {
-	return 0.0f;
+	return 0.0;
 }
 
 int32 UChaosMoverBackendComponent::GetCurrentSimFrame()
@@ -476,7 +476,7 @@ void UChaosMoverBackendComponent::ConsumeOutputData(const FMoverTimeStep& TimeSt
 	SimOutputRecord.Add(TimeStep, OutputData);
 }
 
-void UChaosMoverBackendComponent::FinalizeFrame(float ResultsTimeInMs)
+void UChaosMoverBackendComponent::FinalizeFrame(double ResultsTimeInMs)
 {
 	Chaos::EnsureIsInGameThreadContext();
 
diff --git a/Engine/Plugins/Experimental/ChaosMover/Source/ChaosMover/Private/Backends/ChaosMoverSubsystem.cpp b/Engine/Plugins/Experimental/ChaosMover/Source/ChaosMover/Private/Backends/ChaosMoverSubsystem.cpp
--- a/Engine/Plugins/Experimental/ChaosMover/Source/ChaosMover/Private/Backends/ChaosMoverSubsystem.cpp
+++ b/Engine/Plugins/Experimental/ChaosMover/Source/ChaosMover/Private/Backends/ChaosMoverSubsystem.cpp
@@ -115,7 +115,7 @@ void UChaosMoverSubsystem::InjectInputs_External(int32 PhysicsStep, int32 NumSte
 			{
 				AsyncInput->PhysicsSolver = Solver;
 
-				TimeStep.BaseSimTimeMs = Solver->GetPhysicsResultsTime_External() * 1000.0f;
+				TimeStep.BaseSimTimeMs = Solver->GetPhysicsResultsTime_External() * 1000.0;
 				TimeStep.ServerFrame = Solver->GetCurrentFrame() + NetPhysicsTickOffset;
 
 				if (Solver->IsUsingAsyncResults())
@@ -160,12 +160,12 @@ void UChaosMoverSubsystem::OnPostPhysicsTick(FChaosScene* Scene)
 		Chaos::PhysicsParallelFor(AsyncOutput->Backends.Num(), LambdaParallelUpdate, ForceSingleThread);
 	}
 
-	float ResultsTimeInMs = 0.0f;
+	double ResultsTimeInMs = 0.0;
 	if (Scene)
 	{
 		if (Chaos::FPhysicsSolver* Solver = Scene->GetSolver())
 		{
-			ResultsTimeInMs = Solver->GetPhysicsResultsTime_External() * 1000.0f;
+			ResultsTimeInMs = Solver->GetPhysicsResultsTime_External() * 1000.0;
 		}
 	}
 
diff --git a/Engine/Plugins/Experimental/ChaosMover/Source/ChaosMover/Private/PathedMovement/ChaosPathedMovementMode.cpp b/Engine/Plugins/Experimental/ChaosMover/Source/ChaosMover/Private/PathedMovement/ChaosPathedMovementMode.cpp
--- a/Engine/Plugins/Experimental/ChaosMover/Source/ChaosMover/Private/PathedMovement/ChaosPathedMovementMode.cpp
+++ b/Engine/Plugins/Experimental/ChaosMover/Source/ChaosMover/Private/PathedMovement/ChaosPathedMovementMode.cpp
@@ -314,7 +314,7 @@ void UChaosPathedMovementMode::SimulationTick_Implementation(const FSimulationTi
 #endif
 }
 
-void UChaosPathedMovementMode::OnPlaybackStarting(float BaseSimTimeMs)
+void UChaosPathedMovementMode::OnPlaybackStarting(double BaseSimTimeMs)
 {
 	// Add a pathed movement started event to the simulation (will be broadcast on GT during post sim)
 	if (Simulation)
@@ -330,7 +330,7 @@ void UChaosPathedMovementMode::OnPlaybackStarting(float BaseSimTimeMs)
 	}
 }
 
-void UChaosPathedMovementMode::OnPlaybackStopping(float BaseSimTimeMs, bool bReachedEndOfPlayback)
+void UChaosPathedMovementMode::OnPlaybackStopping(double BaseSimTimeMs, bool bReachedEndOfPlayback)
 {
 	// Add a pathed movement stopped event to the simulation (will be broadcast on GT during post sim)
 	if (Simulation)
@@ -347,7 +347,7 @@ void UChaosPathedMovementMode::OnPlaybackStopping(float BaseSimTimeMs, bool bRea
 	}
 }
 
-void UChaosPathedMovementMode::OnPlaybackBounced(float BaseSimTimeMs)
+void UChaosPathedMovementMode::OnPlaybackBounced(double BaseSimTimeMs)
 {
 	// Add a pathed movement bounced event to the simulation (will be broadcast on GT during post sim)
 	if (Simulation)
@@ -363,7 +363,7 @@ void UChaosPathedMovementMode::OnPlaybackBounced(float BaseSimTimeMs)
 	}
 }
 
-void UChaosPathedMovementMode::OnInReverseChanging(float BaseSimTimeMs, bool bNewInReverse)
+void UChaosPathedMovementMode::OnInReverseChanging(double BaseSimTimeMs, bool bNewInReverse)
 {
 	// Add a pathed movement stopped event to the simulation (will be broadcast on GT during post sim)
 	if (Simulation)
@@ -379,7 +379,7 @@ void UChaosPathedMovementMode::OnInReverseChanging(float BaseSimTimeMs, bool bNe
 	}
 }
 
-void UChaosPathedMovementMode::OnLoopingChanging(float BaseSimTimeMs, bool bNewLooping)
+void UChaosPathedMovementMode::OnLoopingChanging(double BaseSimTimeMs, bool bNewLooping)
 {
 	// Add a pathed movement stopped event to the simulation (will be broadcast on GT during post sim)
 	if (Simulation)
@@ -395,7 +395,7 @@ void UChaosPathedMovementMode::OnLoopingChanging(float BaseSimTimeMs, bool bNewL
 	}
 }
 
-void UChaosPathedMovementMode::OnOneWayChanging(float BaseSimTimeMs, bool bNewOneWay)
+void UChaosPathedMovementMode::OnOneWayChanging(double BaseSimTimeMs, bool bNewOneWay)
 {
 	// Add a pathed movement stopped event to the simulation (will be broadcast on GT during post sim)
 	if (Simulation)
diff --git a/Engine/Plugins/Experimental/ChaosMover/Source/ChaosMover/Public/ChaosMover/Backends/ChaosMoverBackend.h b/Engine/Plugins/Experimental/ChaosMover/Source/ChaosMover/Public/ChaosMover/Backends/ChaosMoverBackend.h
--- a/Engine/Plugins/Experimental/ChaosMover/Source/ChaosMover/Public/ChaosMover/Backends/ChaosMoverBackend.h
+++ b/Engine/Plugins/Experimental/ChaosMover/Source/ChaosMover/Public/ChaosMover/Backends/ChaosMoverBackend.h
@@ -40,12 +40,12 @@ public:
 	// IMoverBackendLiaisonInterface
 	CHAOSMOVER_API virtual bool IsAsync() const override;
 	// Unused IMoverBackendLiaisonInterface API
-	CHAOSMOVER_API virtual float GetCurrentSimTimeMs();
+	CHAOSMOVER_API virtual double GetCurrentSimTimeMs();
 	CHAOSMOVER_API virtual int32 GetCurrentSimFrame();
 
 	CHAOSMOVER_API virtual void ProduceInputData(int32 PhysicsStep, int32 NumSteps, const FMoverTimeStep& TimeStep, UE::ChaosMover::FSimulationInputData& InputData);
 	CHAOSMOVER_API virtual void ConsumeOutputData(const FMoverTimeStep& TimeStep, const UE::ChaosMover::FSimulationOutputData& OutputData);
-	CHAOSMOVER_API virtual void FinalizeFrame(float ResultsTimeInMs);
+	CHAOSMOVER_API virtual void FinalizeFrame(double ResultsTimeInMs);
 
 	UChaosMoverSimulation* GetSimulation();
 	const UChaosMoverSimulation* GetSimulation() const;
diff --git a/Engine/Plugins/Experimental/ChaosMover/Source/ChaosMover/Public/ChaosMover/ChaosMoverSimulationTypes.h b/Engine/Plugins/Experimental/ChaosMover/Source/ChaosMover/Public/ChaosMover/ChaosMoverSimulationTypes.h
--- a/Engine/Plugins/Experimental/ChaosMover/Source/ChaosMover/Public/ChaosMover/ChaosMoverSimulationTypes.h
+++ b/Engine/Plugins/Experimental/ChaosMover/Source/ChaosMover/Public/ChaosMover/ChaosMoverSimulationTypes.h
@@ -226,7 +226,7 @@ struct FStanceModifiedEventData : public FMoverSimulationEventData
 {
 	GENERATED_BODY()
 
-	FStanceModifiedEventData(float InEventTimeMs, EStanceMode InOldStance, EStanceMode InNewStance, FEventProcessedCallbackPtr InEventProcessedCallback = nullptr)
+	FStanceModifiedEventData(double InEventTimeMs, EStanceMode InOldStance, EStanceMode InNewStance, FEventProcessedCallbackPtr InEventProcessedCallback = nullptr)
 		: FMoverSimulationEventData(InEventTimeMs, InEventProcessedCallback)
 		, OldStance(InOldStance)
 		, NewStance(InNewStance)
diff --git a/Engine/Plugins/Experimental/ChaosMover/Source/ChaosMover/Public/ChaosMover/ChaosMoverStateMachine.h b/Engine/Plugins/Experimental/ChaosMover/Source/ChaosMover/Public/ChaosMover/ChaosMoverStateMachine.h
--- a/Engine/Plugins/Experimental/ChaosMover/Source/ChaosMover/Public/ChaosMover/ChaosMoverStateMachine.h
+++ b/Engine/Plugins/Experimental/ChaosMover/Source/ChaosMover/Public/ChaosMover/ChaosMoverStateMachine.h
@@ -110,6 +110,6 @@ namespace UE::ChaosMover
 		bool HasAnyInstantEffectsQueued() const;
 		bool ApplyInstantEffects(const FMoverTickStartData& SubstepStartData, const FMoverTimeStep& SubTimeStep, FMoverSyncState& OutputState);
 	
-		float InternalSimTimeMs = 0.0f;
+		double InternalSimTimeMs = 0.0;
 	};
 }
\ No newline at end of file
diff --git a/Engine/Plugins/Experimental/ChaosMover/Source/ChaosMover/Public/ChaosMover/PathedMovement/ChaosPathedMovementMode.h b/Engine/Plugins/Experimental/ChaosMover/Source/ChaosMover/Public/ChaosMover/PathedMovement/ChaosPathedMovementMode.h
--- a/Engine/Plugins/Experimental/ChaosMover/Source/ChaosMover/Public/ChaosMover/PathedMovement/ChaosPathedMovementMode.h
+++ b/Engine/Plugins/Experimental/ChaosMover/Source/ChaosMover/Public/ChaosMover/PathedMovement/ChaosPathedMovementMode.h
@@ -65,12 +65,12 @@ public:
 	float GetOneWayTripDuration() const { return OneWayPlaybackDuration; }
 	
 protected:
-	void OnPlaybackStarting(float BaseSimTimeMs);
-	void OnPlaybackStopping(float BaseSimTimeMs, bool bReachedEndOfPlayback);
-	void OnPlaybackBounced(float BaseSimTimeMs);
-	void OnInReverseChanging(float BaseSimTimeMs, bool bNewInReverse);
-	void OnLoopingChanging(float BaseSimTimeMs, bool bNewLooping);
-	void OnOneWayChanging(float BaseSimTimeMs, bool bNewOneWay);
+	void OnPlaybackStarting(double BaseSimTimeMs);
+	void OnPlaybackStopping(double BaseSimTimeMs, bool bReachedEndOfPlayback);
+	void OnPlaybackBounced(double BaseSimTimeMs);
+	void OnInReverseChanging(double BaseSimTimeMs, bool bNewInReverse);
+	void OnLoopingChanging(double BaseSimTimeMs, bool bNewLooping);
+	void OnOneWayChanging(double BaseSimTimeMs, bool bNewOneWay);
 
 	void SetDefaultJointConstraintProperties();
 	
diff --git a/Engine/Plugins/Experimental/ChaosMover/Source/ChaosMover/Public/ChaosMover/PathedMovement/ChaosPathedMovementTypes.h b/Engine/Plugins/Experimental/ChaosMover/Source/ChaosMover/Public/ChaosMover/PathedMovement/ChaosPathedMovementTypes.h
--- a/Engine/Plugins/Experimental/ChaosMover/Source/ChaosMover/Public/ChaosMover/PathedMovement/ChaosPathedMovementTypes.h
+++ b/Engine/Plugins/Experimental/ChaosMover/Source/ChaosMover/Public/ChaosMover/PathedMovement/ChaosPathedMovementTypes.h
@@ -215,7 +215,7 @@ struct FPathedMovementStartedEventData : public FMoverSimulationEventData
 {
 	GENERATED_BODY()
 
-	FPathedMovementStartedEventData(float InEventTimeMs)
+	FPathedMovementStartedEventData(double InEventTimeMs)
 		: FMoverSimulationEventData(InEventTimeMs)
 	{
 	}
@@ -232,7 +232,7 @@ struct FPathedMovementStoppedEventData : public FMoverSimulationEventData
 {
 	GENERATED_BODY()
 
-	FPathedMovementStoppedEventData(float InEventTimeMs, bool InbReachedEndOfPlayback)
+	FPathedMovementStoppedEventData(double InEventTimeMs, bool InbReachedEndOfPlayback)
 		: FMoverSimulationEventData(InEventTimeMs)
 		, bReachedEndOfPlayback(InbReachedEndOfPlayback)
 	{
@@ -252,7 +252,7 @@ struct FPathedMovementBouncedEventData : public FMoverSimulationEventData
 {
 	GENERATED_BODY()
 
-	FPathedMovementBouncedEventData(float InEventTimeMs)
+	FPathedMovementBouncedEventData(double InEventTimeMs)
 		: FMoverSimulationEventData(InEventTimeMs)
 	{
 	}
@@ -269,7 +269,7 @@ struct FPathedMovementInReverseChangedEventData : public FMoverSimulationEventDa
 {
 	GENERATED_BODY()
 
-	FPathedMovementInReverseChangedEventData(float InEventTimeMs, bool InbNewInReverse)
+	FPathedMovementInReverseChangedEventData(double InEventTimeMs, bool InbNewInReverse)
 		: FMoverSimulationEventData(InEventTimeMs)
 		, bNewInReverse(InbNewInReverse)
 	{
@@ -289,7 +289,7 @@ struct FPathedMovementLoopingChangedEventData : public FMoverSimulationEventData
 {
 	GENERATED_BODY()
 
-	FPathedMovementLoopingChangedEventData(float InEventTimeMs, bool InbNewLooping)
+	FPathedMovementLoopingChangedEventData(double InEventTimeMs, bool InbNewLooping)
 		: FMoverSimulationEventData(InEventTimeMs)
 		, bNewLooping(InbNewLooping)
 	{
@@ -309,7 +309,7 @@ struct FPathedMovementOneWayChangedEventData : public FMoverSimulationEventData
 {
 	GENERATED_BODY()
 
-	FPathedMovementOneWayChangedEventData(float InEventTimeMs, bool InbNewOneWay)
+	FPathedMovementOneWayChangedEventData(double InEventTimeMs, bool InbNewOneWay)
 		: FMoverSimulationEventData(InEventTimeMs)
 		, bNewOneWay(InbNewOneWay)
 	{
diff --git a/Engine/Plugins/Experimental/Mover/README.md b/Engine/Plugins/Experimental/Mover/README.md
--- a/Engine/Plugins/Experimental/Mover/README.md
+++ b/Engine/Plugins/Experimental/Mover/README.md
@@ -198,3 +198,5 @@ Additionally, when the simulation runs at a lower rate than rendering, the resul
 **Sim Blackboard Is Not Yet Rollback-Friendly:** When rollbacks occur, the blackboard's contents are simply invalidated. Although the blackboard is useful for avoiding repeating computations from a prior frame, movement logic cannot always rely on it having a valid entry and should have a fallback.
 
 **Multiple Modifiers of the Same Type:** Currently Mover expects there to be only one modifier of a given type active at a time. This can be worked around by overriding the Matches function and use more criteria than just the type, as long as that criteria is reflected by data included in the NetSerialize function.
+
+**Absolute Timestamps in Blackboard (prior to 5.7):** Timestamps were originally kept as single-precision floats, and later updated to double-precision. FMoverTimeStep::BaseSimTimeMs is an example. If you had been storing any timestamps in a Mover blackboard, be aware that inconsistent use of float and double as the element type can lead to reading bad data. If you 'Set' a blackboard item using a double, you need to 'Get' it as a double as well.
diff --git a/Engine/Plugins/Experimental/Mover/Source/Mover/Private/Backends/MoverNetworkPhysicsLiaisonBase.cpp b/Engine/Plugins/Experimental/Mover/Source/Mover/Private/Backends/MoverNetworkPhysicsLiaisonBase.cpp
--- a/Engine/Plugins/Experimental/Mover/Source/Mover/Private/Backends/MoverNetworkPhysicsLiaisonBase.cpp
+++ b/Engine/Plugins/Experimental/Mover/Source/Mover/Private/Backends/MoverNetworkPhysicsLiaisonBase.cpp
@@ -221,14 +221,14 @@ UMoverNetworkPhysicsLiaisonComponentBase::UMoverNetworkPhysicsLiaisonComponentBa
 //////////////////////////////////////////////////////////////////////////
 //  UMoverNetworkPhysicsLiaisonComponentBase IMoverBackendLiaisonInterface
 
-float UMoverNetworkPhysicsLiaisonComponentBase::GetCurrentSimTimeMs()
+double UMoverNetworkPhysicsLiaisonComponentBase::GetCurrentSimTimeMs()
 {
 	// Note: this is implicitly an _External function
 	if (const Chaos::FPhysicsSolver* Solver = GetPhysicsSolver())
 	{
-		return bUsingAsyncPhysics ? Solver->GetAsyncDeltaTime() * GetCurrentSimFrame() * 1000.0f : Solver->GetSolverTime() * 1000.0f;
+		return bUsingAsyncPhysics ? Solver->GetAsyncDeltaTime() * GetCurrentSimFrame() * 1000.0 : Solver->GetSolverTime() * 1000.0;
 	}
-	return 0.0f;
+	return 0.0;
 }
 
 int32 UMoverNetworkPhysicsLiaisonComponentBase::GetCurrentSimFrame()
@@ -448,7 +448,7 @@ FMoverTimeStep UMoverNetworkPhysicsLiaisonComponentBase::GetCurrentAsyncMoverTim
 	{
 		TimeStep.ServerFrame = Solver->GetCurrentFrame() + GetNetworkPhysicsTickOffset_Internal();
 		TimeStep.StepMs = Solver->GetAsyncDeltaTime() * 1000.0f;
-		TimeStep.BaseSimTimeMs = TimeStep.ServerFrame * TimeStep.StepMs;
+		TimeStep.BaseSimTimeMs = (double)TimeStep.StepMs * TimeStep.ServerFrame;
 		TimeStep.bIsResimulating = Solver->GetEvolution()->IsResimming();
 	}
 
@@ -466,7 +466,7 @@ FMoverTimeStep UMoverNetworkPhysicsLiaisonComponentBase::GetCurrentAsyncMoverTim
 		const int32 Offset = GetNetworkPhysicsTickOffset_External();
 		TimeStep.ServerFrame = Solver->GetCurrentFrame() + Offset;
 		TimeStep.StepMs = Solver->GetAsyncDeltaTime() * 1000.0f;
-		TimeStep.BaseSimTimeMs = Solver->GetPhysicsResultsTime_External() * 1000.0f + Offset * TimeStep.StepMs;
+		TimeStep.BaseSimTimeMs = (Solver->GetPhysicsResultsTime_External() * 1000.0) + ((double)TimeStep.StepMs * Offset);
 		TimeStep.bIsResimulating = Solver->GetEvolution()->IsResimming();
 	}
 
@@ -482,7 +482,7 @@ FMoverTimeStep UMoverNetworkPhysicsLiaisonComponentBase::GetCurrentMoverTimeStep
 	{
 		TimeStep.ServerFrame = Solver->GetCurrentFrame();
 		TimeStep.StepMs = DeltaSeconds * 1000.0f;
-		TimeStep.BaseSimTimeMs = Solver->GetSolverTime() * 1000.0f;
+		TimeStep.BaseSimTimeMs = Solver->GetSolverTime() * 1000.0;
 		TimeStep.bIsResimulating = Solver->GetEvolution()->IsResimming();
 	}
 
diff --git a/Engine/Plugins/Experimental/Mover/Source/Mover/Private/Backends/MoverNetworkPredictionLiaison.cpp b/Engine/Plugins/Experimental/Mover/Source/Mover/Private/Backends/MoverNetworkPredictionLiaison.cpp
--- a/Engine/Plugins/Experimental/Mover/Source/Mover/Private/Backends/MoverNetworkPredictionLiaison.cpp
+++ b/Engine/Plugins/Experimental/Mover/Source/Mover/Private/Backends/MoverNetworkPredictionLiaison.cpp
@@ -103,14 +103,20 @@ void UMoverNetworkPredictionLiaisonComponent::SimulationTick(const FNetSimTimeSt
 		StartData.SyncState.SyncStateCollection.FindOrAddDataByType(PersistentSyncEntry.RequiredType);
 	}
 	
-	MoverComp->SimulationTick(FMoverTimeStep(TimeStep), StartData, OUT EndData);
+	FMoverTimeStep MoverTimeStep;
+
+	MoverTimeStep.ServerFrame	= TimeStep.Frame;
+	MoverTimeStep.BaseSimTimeMs = TimeStep.TotalSimulationTime;
+	MoverTimeStep.StepMs		= TimeStep.StepMS;
+
+	MoverComp->SimulationTick(MoverTimeStep, StartData, OUT EndData);
 
 	*SimOutput.Sync = EndData.SyncState;
     *SimOutput.Aux.Get() = EndData.AuxState;
 }
 
 
-float UMoverNetworkPredictionLiaisonComponent::GetCurrentSimTimeMs()
+double UMoverNetworkPredictionLiaisonComponent::GetCurrentSimTimeMs()
 {
 	return NetworkPredictionProxy.GetTotalSimTimeMS();
 }
diff --git a/Engine/Plugins/Experimental/Mover/Source/Mover/Private/Backends/MoverStandaloneLiaison.cpp b/Engine/Plugins/Experimental/Mover/Source/Mover/Private/Backends/MoverStandaloneLiaison.cpp
--- a/Engine/Plugins/Experimental/Mover/Source/Mover/Private/Backends/MoverStandaloneLiaison.cpp
+++ b/Engine/Plugins/Experimental/Mover/Source/Mover/Private/Backends/MoverStandaloneLiaison.cpp
@@ -87,7 +87,7 @@ UMoverStandaloneLiaisonComponent::UMoverStandaloneLiaisonComponent(const FObject
 
 }
 
-float UMoverStandaloneLiaisonComponent::GetCurrentSimTimeMs()
+double UMoverStandaloneLiaisonComponent::GetCurrentSimTimeMs()
 {
 	return CurrentSimTimeMs;
 }
diff --git a/Engine/Plugins/Experimental/Mover/Source/Mover/Private/Debug/MoverDebugComponent.cpp b/Engine/Plugins/Experimental/Mover/Source/Mover/Private/Debug/MoverDebugComponent.cpp
--- a/Engine/Plugins/Experimental/Mover/Source/Mover/Private/Debug/MoverDebugComponent.cpp
+++ b/Engine/Plugins/Experimental/Mover/Source/Mover/Private/Debug/MoverDebugComponent.cpp
@@ -137,9 +137,9 @@ void UMoverDebugComponent::InitHistoryTracking()
 
 void UMoverDebugComponent::UpdateHistoryTrackingForFrame(const FMoverTimeStep& TimeStep, const FMoverSyncState* SyncState, const FMoverAuxStateContext* AuxState)
 {
-	const float CurrentSimTimeMs = TimeStep.BaseSimTimeMs;
-	const float MaxTimeBetweenSamplesMs = (1000.0f / HistorySamplesPerSecond);
-	const float MinSimTimeOfSampleMs = CurrentSimTimeMs - (HistoryTrackingSeconds * 1000.f);
+	const double CurrentSimTimeMs = TimeStep.BaseSimTimeMs;
+	const double MaxTimeBetweenSamplesMs = (1000.0 / HistorySamplesPerSecond);
+	const double MinSimTimeOfSampleMs = CurrentSimTimeMs - (HistoryTrackingSeconds * 1000.0);
 
 	// Check if it's time to kick any samples out
 	HistorySamples.RemoveAll([&](const FTrajectorySampleInfo& Sample)
@@ -181,7 +181,7 @@ void UMoverDebugComponent::UpdateHistoryTrackingForFrame(const FMoverTimeStep& T
 
 void UMoverDebugComponent::OnHistoryTrackingRollback(const FMoverTimeStep& NewTimeStep, const FMoverTimeStep& InvalidatedTimeStep)
 {
-	const float CurrentSimTimeMs = NewTimeStep.BaseSimTimeMs;
+	const double CurrentSimTimeMs = NewTimeStep.BaseSimTimeMs;
 
 	// Remove any samples that are newer than the current sim time. This implies they've been rolled back.
 	HistorySamples.RemoveAll([&](const FTrajectorySampleInfo& Sample)
diff --git a/Engine/Plugins/Experimental/Mover/Source/Mover/Private/DefaultMovementSet/LayeredMoves/AnimRootMotionLayeredMove.cpp b/Engine/Plugins/Experimental/Mover/Source/Mover/Private/DefaultMovementSet/LayeredMoves/AnimRootMotionLayeredMove.cpp
--- a/Engine/Plugins/Experimental/Mover/Source/Mover/Private/DefaultMovementSet/LayeredMoves/AnimRootMotionLayeredMove.cpp
+++ b/Engine/Plugins/Experimental/Mover/Source/Mover/Private/DefaultMovementSet/LayeredMoves/AnimRootMotionLayeredMove.cpp
@@ -60,8 +60,8 @@ bool FLayeredMove_AnimRootMotion::GenerateMove(const FMoverTickStartData& SimSta
 	//     - detecting if the montage instance is interrupted and attempting to interrupt and scheduling this move to end at the same sim time
 	
 	// Note that Montage 'position' equates to seconds when PlayRate is 1
-	const float SecondsSinceMontageStarted = (TimeStep.BaseSimTimeMs - StartSimTimeMs) / 1000.f;
-	const float ScaledSecondsSinceMontageStarted = SecondsSinceMontageStarted * PlayRate;
+	const double SecondsSinceMontageStarted = (TimeStep.BaseSimTimeMs - StartSimTimeMs) / 1000.0;
+	const double ScaledSecondsSinceMontageStarted = SecondsSinceMontageStarted * PlayRate;
 
 	const float ExtractionStartPosition = StartingMontagePosition + ScaledSecondsSinceMontageStarted;
 	const float ExtractionEndPosition   = ExtractionStartPosition + (DeltaSeconds * PlayRate);
diff --git a/Engine/Plugins/Experimental/Mover/Source/Mover/Private/DefaultMovementSet/LayeredMoves/BasicLayeredMoves.cpp b/Engine/Plugins/Experimental/Mover/Source/Mover/Private/DefaultMovementSet/LayeredMoves/BasicLayeredMoves.cpp
--- a/Engine/Plugins/Experimental/Mover/Source/Mover/Private/DefaultMovementSet/LayeredMoves/BasicLayeredMoves.cpp
+++ b/Engine/Plugins/Experimental/Mover/Source/Mover/Private/DefaultMovementSet/LayeredMoves/BasicLayeredMoves.cpp
@@ -46,7 +46,7 @@ bool FLayeredMove_LinearVelocity::GenerateMove(const FMoverTickStartData& SimSta
 
 	if (MagnitudeOverTime && DurationMs > 0)
 	{
-		const float TimeValue = DurationMs > 0.f ? FMath::Clamp((TimeStep.BaseSimTimeMs - StartSimTimeMs) / DurationMs, 0.f, 1.f) : TimeStep.BaseSimTimeMs;
+		const double TimeValue = FMath::Clamp((TimeStep.BaseSimTimeMs - StartSimTimeMs) / DurationMs, 0.0, 1.0);
 		const float TimeFactor = MagnitudeOverTime->GetFloatValue(TimeValue);
 		VelocityThisFrame *= TimeFactor;
 	}
@@ -525,9 +525,9 @@ bool FLayeredMove_RadialImpulse::GenerateMove(const FMoverTickStartData& StartSt
 				AdditiveMagnitudeFactor -= (1.f - DistanceFactor);
 			}
 
-			if (MagnitudeOverTime)
+			if (MagnitudeOverTime && DurationMs > 0)
 			{
-				const float TimeValue = DurationMs > 0.f ? FMath::Clamp(TimeStep.BaseSimTimeMs / DurationMs, 0.f, 1.f) : TimeStep.BaseSimTimeMs;
+				const double TimeValue = FMath::Clamp((TimeStep.BaseSimTimeMs - StartSimTimeMs) / DurationMs, 0.0, 1.0);
 				const float TimeFactor = MagnitudeOverTime->GetFloatValue(TimeValue);
 				AdditiveMagnitudeFactor -= (1.f - TimeFactor);
 			}
diff --git a/Engine/Plugins/Experimental/Mover/Source/Mover/Private/LayeredMove.cpp b/Engine/Plugins/Experimental/Mover/Source/Mover/Private/LayeredMove.cpp
--- a/Engine/Plugins/Experimental/Mover/Source/Mover/Private/LayeredMove.cpp
+++ b/Engine/Plugins/Experimental/Mover/Source/Mover/Private/LayeredMove.cpp
@@ -135,8 +135,8 @@ void FLayeredMoveGroup::QueueLayeredMove(TSharedPtr<FLayeredMoveBase> Move)
 
 TArray<TSharedPtr<FLayeredMoveBase>> FLayeredMoveGroup::GenerateActiveMoves(const FMoverTimeStep& TimeStep, const UMoverComponent* MoverComp, UMoverBlackboard* SimBlackboard)
 {
-	const float SimStartTimeMs		= TimeStep.BaseSimTimeMs;
-	const float SimTimeAfterTickMs	= SimStartTimeMs + TimeStep.StepMs;
+	const double SimStartTimeMs		= TimeStep.BaseSimTimeMs;
+	const double SimTimeAfterTickMs	= SimStartTimeMs + TimeStep.StepMs;
 
 	FlushMoveArrays(MoverComp, SimBlackboard, SimStartTimeMs, /*bIsAsync =*/ false);
 
@@ -145,8 +145,8 @@ TArray<TSharedPtr<FLayeredMoveBase>> FLayeredMoveGroup::GenerateActiveMoves(cons
 
 TArray<TSharedPtr<FLayeredMoveBase>> FLayeredMoveGroup::GenerateActiveMoves_Async(const FMoverTimeStep& TimeStep, UMoverBlackboard* SimBlackboard)
 {
-	const float SimStartTimeMs		= TimeStep.BaseSimTimeMs;
-	const float SimTimeAfterTickMs	= SimStartTimeMs + TimeStep.StepMs;
+	const double SimStartTimeMs		= TimeStep.BaseSimTimeMs;
+	const double SimTimeAfterTickMs	= SimStartTimeMs + TimeStep.StepMs;
 
 	FlushMoveArrays(nullptr, SimBlackboard, SimStartTimeMs, /*bIsAsync =*/ true);
 
@@ -362,7 +362,7 @@ const FLayeredMoveBase* FLayeredMoveGroup::FindQueuedMove(const UScriptStruct* L
 	return nullptr;
 }
 
-void FLayeredMoveGroup::FlushMoveArrays(const UMoverComponent* MoverComp, UMoverBlackboard* SimBlackboard, float CurrentSimTimeMs, bool bIsAsync)
+void FLayeredMoveGroup::FlushMoveArrays(const UMoverComponent* MoverComp, UMoverBlackboard* SimBlackboard, double CurrentSimTimeMs, bool bIsAsync)
 {
 	if (bIsAsync)
 	{
diff --git a/Engine/Plugins/Experimental/Mover/Source/Mover/Private/MovementModifier.cpp b/Engine/Plugins/Experimental/Mover/Source/Mover/Private/MovementModifier.cpp
--- a/Engine/Plugins/Experimental/Mover/Source/Mover/Private/MovementModifier.cpp
+++ b/Engine/Plugins/Experimental/Mover/Source/Mover/Private/MovementModifier.cpp
@@ -49,7 +49,7 @@ void FMovementModifierBase::EndModifier_Async(const FMovementModifierParams_Asyn
 	OnEnd_Async(Params);
 }
 
-bool FMovementModifierBase::IsFinished(float CurrentSimTimeMs) const
+bool FMovementModifierBase::IsFinished(double CurrentSimTimeMs) const
 {
 	const bool bHasStarted = (StartSimTimeMs >= 0.f);
 	const bool bTimeExpired = bHasStarted && (DurationMs > 0.f) && (StartSimTimeMs + DurationMs <= CurrentSimTimeMs);
diff --git a/Engine/Plugins/Experimental/Mover/Source/Mover/Private/MoverSimulationTypes.cpp b/Engine/Plugins/Experimental/Mover/Source/Mover/Private/MoverSimulationTypes.cpp
--- a/Engine/Plugins/Experimental/Mover/Source/Mover/Private/MoverSimulationTypes.cpp
+++ b/Engine/Plugins/Experimental/Mover/Source/Mover/Private/MoverSimulationTypes.cpp
@@ -18,7 +18,7 @@ namespace UE::Mover
 		Events.Empty();
 	}
 
-	void FSimulationOutputData::Interpolate(const FSimulationOutputData& From, const FSimulationOutputData& To, float Alpha, float SimTimeMs)
+	void FSimulationOutputData::Interpolate(const FSimulationOutputData& From, const FSimulationOutputData& To, float Alpha, double SimTimeMs)
 	{
 		SyncState.Interpolate(&From.SyncState, &To.SyncState, Alpha);
 		LastUsedInputCmd.Interpolate(&From.LastUsedInputCmd, &To.LastUsedInputCmd, Alpha);
@@ -66,11 +66,11 @@ namespace UE::Mover
 		return Data[CurrentIndex].SimOutputData;
 	}
 
-	void FSimulationOutputRecord::CreateInterpolatedResult(float AtBaseTimeMs, FMoverTimeStep& OutTimeStep, FSimulationOutputData& OutData)
+	void FSimulationOutputRecord::CreateInterpolatedResult(double AtBaseTimeMs, FMoverTimeStep& OutTimeStep, FSimulationOutputData& OutData)
 	{
 		const uint8 PrevIndex = (CurrentIndex + 1) % 2;
-		const float PrevTimeMs = Data[PrevIndex].TimeStep.BaseSimTimeMs;
-		const float CurrTimeMs = Data[CurrentIndex].TimeStep.BaseSimTimeMs;
+		const double PrevTimeMs = Data[PrevIndex].TimeStep.BaseSimTimeMs;
+		const double CurrTimeMs = Data[CurrentIndex].TimeStep.BaseSimTimeMs;
 
 		if (FMath::IsNearlyEqual(PrevTimeMs, CurrTimeMs) || (AtBaseTimeMs >= CurrTimeMs))
 		{
diff --git a/Engine/Plugins/Experimental/Mover/Source/Mover/Public/Backends/MoverBackendLiaison.h b/Engine/Plugins/Experimental/Mover/Source/Mover/Public/Backends/MoverBackendLiaison.h
--- a/Engine/Plugins/Experimental/Mover/Source/Mover/Public/Backends/MoverBackendLiaison.h
+++ b/Engine/Plugins/Experimental/Mover/Source/Mover/Public/Backends/MoverBackendLiaison.h
@@ -26,7 +26,7 @@ class IMoverBackendLiaisonInterface
 	GENERATED_BODY()
 
 public:
-	virtual float GetCurrentSimTimeMs() = 0;
+	virtual double GetCurrentSimTimeMs() = 0;
 	virtual int32 GetCurrentSimFrame() = 0;
 
 	// Whether this backend will simulate movement asynchronously
diff --git a/Engine/Plugins/Experimental/Mover/Source/Mover/Public/Backends/MoverNetworkPhysicsLiaisonBase.h b/Engine/Plugins/Experimental/Mover/Source/Mover/Public/Backends/MoverNetworkPhysicsLiaisonBase.h
--- a/Engine/Plugins/Experimental/Mover/Source/Mover/Public/Backends/MoverNetworkPhysicsLiaisonBase.h
+++ b/Engine/Plugins/Experimental/Mover/Source/Mover/Public/Backends/MoverNetworkPhysicsLiaisonBase.h
@@ -191,7 +191,7 @@ public:
 	UE_API UMoverNetworkPhysicsLiaisonComponentBase();
 
 	// IMoverBackendLiaisonInterface
-	UE_API virtual float GetCurrentSimTimeMs() override;
+	UE_API virtual double GetCurrentSimTimeMs() override;
 	UE_API virtual int32 GetCurrentSimFrame() override;
 #if WITH_EDITOR
 	UE_API virtual EDataValidationResult ValidateData(FDataValidationContext& Context, const UMoverComponent& ValidationMoverComp) const override;
diff --git a/Engine/Plugins/Experimental/Mover/Source/Mover/Public/Backends/MoverNetworkPredictionLiaison.h b/Engine/Plugins/Experimental/Mover/Source/Mover/Public/Backends/MoverNetworkPredictionLiaison.h
--- a/Engine/Plugins/Experimental/Mover/Source/Mover/Public/Backends/MoverNetworkPredictionLiaison.h
+++ b/Engine/Plugins/Experimental/Mover/Source/Mover/Public/Backends/MoverNetworkPredictionLiaison.h
@@ -7,6 +7,7 @@
 #include "Backends/MoverBackendLiaison.h"
 #include "NetworkPredictionComponent.h"
 #include "NetworkPredictionSimulation.h"
+#include "NetworkPredictionTickState.h"
 #include "MovementMode.h"
 #include "MoverTypes.h"
 
@@ -50,7 +51,7 @@ public:
 	// End NP Driver interface
 
 	// IMoverBackendLiaisonInterface
-	UE_API virtual float GetCurrentSimTimeMs() override;
+	UE_API virtual double GetCurrentSimTimeMs() override;
 	UE_API virtual int32 GetCurrentSimFrame() override;
 	UE_API virtual bool ReadPendingSyncState(OUT FMoverSyncState& OutSyncState) override;
 	UE_API virtual bool WritePendingSyncState(const FMoverSyncState& SyncStateToWrite) override;
diff --git a/Engine/Plugins/Experimental/Mover/Source/Mover/Public/Backends/MoverStandaloneLiaison.h b/Engine/Plugins/Experimental/Mover/Source/Mover/Public/Backends/MoverStandaloneLiaison.h
--- a/Engine/Plugins/Experimental/Mover/Source/Mover/Public/Backends/MoverStandaloneLiaison.h
+++ b/Engine/Plugins/Experimental/Mover/Source/Mover/Public/Backends/MoverStandaloneLiaison.h
@@ -103,7 +103,7 @@ public:
 	UE_API UMoverStandaloneLiaisonComponent(const FObjectInitializer& ObjectInitializer = FObjectInitializer::Get());
 
 	// IMoverBackendLiaisonInterface
-	UE_API virtual float GetCurrentSimTimeMs() override;
+	UE_API virtual double GetCurrentSimTimeMs() override;
 	UE_API virtual int32 GetCurrentSimFrame() override;
 	UE_API virtual bool ReadPendingSyncState(OUT FMoverSyncState& OutSyncState) override;
 	UE_API virtual bool WritePendingSyncState(const FMoverSyncState& SyncStateToWrite) override;
diff --git a/Engine/Plugins/Experimental/Mover/Source/Mover/Public/DefaultMovementSet/CharacterMoverSimulationTypes.h b/Engine/Plugins/Experimental/Mover/Source/Mover/Public/DefaultMovementSet/CharacterMoverSimulationTypes.h
--- a/Engine/Plugins/Experimental/Mover/Source/Mover/Public/DefaultMovementSet/CharacterMoverSimulationTypes.h
+++ b/Engine/Plugins/Experimental/Mover/Source/Mover/Public/DefaultMovementSet/CharacterMoverSimulationTypes.h
@@ -12,7 +12,7 @@ struct FLandedEventData : public FMoverSimulationEventData
 {
 	GENERATED_BODY()
 
-	FLandedEventData(float InEventTimeMs, const FHitResult& InHitResult, const FName InNewModeName)
+	FLandedEventData(double InEventTimeMs, const FHitResult& InHitResult, const FName InNewModeName)
 		: FMoverSimulationEventData(InEventTimeMs)
 		, HitResult(InHitResult)
 		, NewModeName(InNewModeName)
diff --git a/Engine/Plugins/Experimental/Mover/Source/Mover/Public/DefaultMovementSet/LayeredMoves/MultiJumpLayeredMove.h b/Engine/Plugins/Experimental/Mover/Source/Mover/Public/DefaultMovementSet/LayeredMoves/MultiJumpLayeredMove.h
--- a/Engine/Plugins/Experimental/Mover/Source/Mover/Public/DefaultMovementSet/LayeredMoves/MultiJumpLayeredMove.h
+++ b/Engine/Plugins/Experimental/Mover/Source/Mover/Public/DefaultMovementSet/LayeredMoves/MultiJumpLayeredMove.h
@@ -50,7 +50,7 @@ protected:
 	int32 JumpsInAirRemaining;
 
 	UPROPERTY(BlueprintReadOnly, EditAnywhere, Category = "Jumping")
-	float TimeOfLastJumpMS;
+	double TimeOfLastJumpMS;
 
 	UE_API bool PerformJump(const FMoverDefaultSyncState* SyncState, const FMoverTimeStep& TimeStep, const UMoverComponent* MoverComp, FProposedMove& OutProposedMove);
 };
diff --git a/Engine/Plugins/Experimental/Mover/Source/Mover/Public/LayeredMove.h b/Engine/Plugins/Experimental/Mover/Source/Mover/Public/LayeredMove.h
--- a/Engine/Plugins/Experimental/Mover/Source/Mover/Public/LayeredMove.h
+++ b/Engine/Plugins/Experimental/Mover/Source/Mover/Public/LayeredMove.h
@@ -93,7 +93,7 @@ struct FLayeredMoveBase
 
 	// The simulation time this move first ticked (< 0 means it hasn't started yet)
 	UPROPERTY(VisibleInstanceOnly, BlueprintReadOnly, Transient, Category = Mover)
-	float StartSimTimeMs;
+	double StartSimTimeMs;
 
 	// Settings related to velocity applied to the actor after a layered move has finished
 	UPROPERTY(EditAnywhere, BlueprintReadWrite, Category = Mover)
@@ -242,7 +242,7 @@ struct FLayeredMoveGroup
 
 protected:
 	// Clears out any finished or invalid active moves and adds any queued moves to the active moves
-	UE_API void FlushMoveArrays(const UMoverComponent* MoverComp, UMoverBlackboard* SimBlackboard, float CurrentSimTimeMs, bool bIsAsync);
+	UE_API void FlushMoveArrays(const UMoverComponent* MoverComp, UMoverBlackboard* SimBlackboard, double CurrentSimTimeMs, bool bIsAsync);
 
 	// Helper function for gathering any residual velocity settings from layered moves that just ended
 	UE_API void GatherResidualVelocitySettings(const TSharedPtr<FLayeredMoveBase>& Move, bool& bResidualVelocityOverriden, bool& bClampVelocityOverriden);
diff --git a/Engine/Plugins/Experimental/Mover/Source/Mover/Public/MoveLibrary/MovementUtils.h b/Engine/Plugins/Experimental/Mover/Source/Mover/Public/MoveLibrary/MovementUtils.h
--- a/Engine/Plugins/Experimental/Mover/Source/Mover/Public/MoveLibrary/MovementUtils.h
+++ b/Engine/Plugins/Experimental/Mover/Source/Mover/Public/MoveLibrary/MovementUtils.h
@@ -47,7 +47,7 @@ struct FTrajectorySampleInfo
 
 	// Time stamp of this sample, in server simulation time
 	UPROPERTY(BlueprintReadWrite, EditAnywhere, Category = Mover)
-	float SimTimeMs = 0.0f;
+	double SimTimeMs = 0.0f;
 };
 
 // Input parameters for compute velocity function
diff --git a/Engine/Plugins/Experimental/Mover/Source/Mover/Public/MovementModifier.h b/Engine/Plugins/Experimental/Mover/Source/Mover/Public/MovementModifier.h
--- a/Engine/Plugins/Experimental/Mover/Source/Mover/Public/MovementModifier.h
+++ b/Engine/Plugins/Experimental/Mover/Source/Mover/Public/MovementModifier.h
@@ -113,7 +113,7 @@ struct FMovementModifierBase
 
 	// The simulation time this move first ticked (< 0 means it hasn't started yet)
 	UPROPERTY(VisibleInstanceOnly, BlueprintReadOnly, Transient, Category = Mover)
-	float StartSimTimeMs;
+	double StartSimTimeMs;
 	
 	/** Fired when this modifier is activated. */
 	virtual void OnStart(UMoverComponent* MoverComp, const FMoverTimeStep& TimeStep, const FMoverSyncState& SyncState, const FMoverAuxStateContext& AuxState) {}
@@ -136,7 +136,7 @@ struct FMovementModifierBase
 	UE_API void EndModifier(UMoverComponent* MoverComp, const FMoverTimeStep& TimeStep, const FMoverSyncState& SyncState, const FMoverAuxStateContext& AuxState);
 
 	/** Runtime query whether this modifier is finished and can be destroyed. The default implementation is based on DurationMs. */
-	UE_API virtual bool IsFinished(float CurrentSimTimeMs) const;
+	UE_API virtual bool IsFinished(double CurrentSimTimeMs) const;
 	
 
 	// Begin Async API
diff --git a/Engine/Plugins/Experimental/Mover/Source/Mover/Public/MoverSimulationTypes.h b/Engine/Plugins/Experimental/Mover/Source/Mover/Public/MoverSimulationTypes.h
--- a/Engine/Plugins/Experimental/Mover/Source/Mover/Public/MoverSimulationTypes.h
+++ b/Engine/Plugins/Experimental/Mover/Source/Mover/Public/MoverSimulationTypes.h
@@ -421,7 +421,7 @@ struct FMoverSimulationEventData
 
 	using FEventProcessedCallbackPtr = std::function<void(const FMoverSimulationEventData& Data)>;
 
-	FMoverSimulationEventData(float InEventTimeMs, FEventProcessedCallbackPtr InEventProcessedCallback = nullptr)
+	FMoverSimulationEventData(double InEventTimeMs, FEventProcessedCallbackPtr InEventProcessedCallback = nullptr)
 		: EventProcessedCallback(InEventProcessedCallback)
 		, EventTimeMs(InEventTimeMs)
 	{
@@ -463,7 +463,7 @@ private:
 	FEventProcessedCallbackPtr EventProcessedCallback = nullptr;
 
 public:
-	float EventTimeMs = 0.0f;
+	double EventTimeMs = 0.0;
 };
 
 USTRUCT()
@@ -493,7 +493,7 @@ namespace UE::Mover
 	struct FSimulationOutputData
 	{
 		MOVER_API void Reset();
-		MOVER_API void Interpolate(const FSimulationOutputData& From, const FSimulationOutputData& To, float Alpha, float SimTimeMs);
+		MOVER_API void Interpolate(const FSimulationOutputData& From, const FSimulationOutputData& To, float Alpha, double SimTimeMs);
 
 		FMoverSyncState SyncState;
 		FMoverInputCmdContext LastUsedInputCmd;
@@ -517,7 +517,7 @@ namespace UE::Mover
 		MOVER_API const FSimulationOutputData& GetLatest() const;
 
 		/** This will create an interpolated output and extract events from the stored data with time stamps up until the input time */
-		MOVER_API void CreateInterpolatedResult(float AtBaseTimeMs, FMoverTimeStep& OutTimeStep, FSimulationOutputData& OutData);
+		MOVER_API void CreateInterpolatedResult(double AtBaseTimeMs, FMoverTimeStep& OutTimeStep, FSimulationOutputData& OutData);
 
 		MOVER_API void Clear();
 
diff --git a/Engine/Plugins/Experimental/Mover/Source/Mover/Public/MoverTypes.h b/Engine/Plugins/Experimental/Mover/Source/Mover/Public/MoverTypes.h
--- a/Engine/Plugins/Experimental/Mover/Source/Mover/Public/MoverTypes.h
+++ b/Engine/Plugins/Experimental/Mover/Source/Mover/Public/MoverTypes.h
@@ -4,7 +4,6 @@
 
 #include "Misc/StringBuilder.h"
 #include "Engine/HitResult.h"
-#include "NetworkPredictionTickState.h"
 #include "MoverLog.h"
 #include "Kismet/BlueprintFunctionLibrary.h"
 #include "NativeGameplayTags.h"
@@ -88,23 +87,6 @@ struct FMoverTimeStep
 {
 	GENERATED_BODY()
 	
-	FMoverTimeStep() 
-	{ 
-		bIsResimulating=false; 
-	}
-	
-	FMoverTimeStep(const FNetSimTimeStep& InNetSimTimeStep, bool InIsResimulating)
-		: ServerFrame(InNetSimTimeStep.Frame)
-		, BaseSimTimeMs(InNetSimTimeStep.TotalSimulationTime)
-		, StepMs(InNetSimTimeStep.StepMS)
-		, bIsResimulating(InIsResimulating)
-	{
-	}
-
-	FMoverTimeStep(const FNetSimTimeStep& InNetSimTimeStep)
-		: FMoverTimeStep(InNetSimTimeStep, false)
-	{
-	}
 
 	// The server simulation frame this timestep is associated with
 	UPROPERTY(EditAnywhere, BlueprintReadWrite, Category=Mover)
@@ -112,14 +94,14 @@ struct FMoverTimeStep
 
 	// Starting simulation time (in server simulation timespace)
 	UPROPERTY(EditAnywhere, BlueprintReadWrite, Category=Mover)
-	float BaseSimTimeMs=-1.f;
+	double BaseSimTimeMs=-1.0;
 
 	// The delta time step for this tick
 	UPROPERTY(EditAnywhere, BlueprintReadWrite, Category=Mover)
-	float StepMs=1.f;
+	float StepMs=1.0;
 
 	// Indicates whether this time step is re-simulating based on prior inputs, such as during a correction
-	uint8 bIsResimulating : 1;
+	uint8 bIsResimulating : 1 = 0;
 
 };
 
