[ChaosMover] Clean up and misc. changes

diff --git a/Engine/Plugins/Experimental/ChaosMover/Source/ChaosMover/Private/Character/Transitions/ChaosCharacterLandingCheck.cpp b/Engine/Plugins/Experimental/ChaosMover/Source/ChaosMover/Private/Character/Transitions/ChaosCharacterLandingCheck.cpp
--- a/Engine/Plugins/Experimental/ChaosMover/Source/ChaosMover/Private/Character/Transitions/ChaosCharacterLandingCheck.cpp
+++ b/Engine/Plugins/Experimental/ChaosMover/Source/ChaosMover/Private/Character/Transitions/ChaosCharacterLandingCheck.cpp
@@ -8,7 +8,6 @@
 #include "ChaosMover/Utilities/ChaosGroundMovementUtils.h"
 #include "DefaultMovementSet/CharacterMoverSimulationTypes.h"
 #include "MoveLibrary/FloorQueryUtils.h"
-#include "MoveLibrary/WaterMovementUtils.h"
 #include "MoverSimulationTypes.h"
 
 #include UE_INLINE_GENERATED_CPP_BY_NAME(ChaosCharacterLandingCheck)
diff --git a/Engine/Plugins/Experimental/ChaosMover/Source/ChaosMover/Public/ChaosMover/Character/Modes/ChaosSwimmingMode.h b/Engine/Plugins/Experimental/ChaosMover/Source/ChaosMover/Public/ChaosMover/Character/Modes/ChaosSwimmingMode.h
--- a/Engine/Plugins/Experimental/ChaosMover/Source/ChaosMover/Public/ChaosMover/Character/Modes/ChaosSwimmingMode.h
+++ b/Engine/Plugins/Experimental/ChaosMover/Source/ChaosMover/Public/ChaosMover/Character/Modes/ChaosSwimmingMode.h
@@ -94,7 +94,7 @@ public:
 
 	// Ideal depth when in water. Measured from the center of the collision shape.
 	UPROPERTY(EditAnywhere, BlueprintReadOnly, Category = "Water Control")
-	float SwimmingIdealImmersionDepth = 40.f;
+	float SwimmingIdealImmersionDepth = 43.f;
 
 	UPROPERTY(EditAnywhere, BlueprintReadOnly, Category = "Water Control")
 	FSwimmingSettings SurfaceSwimmingWaterControlSettings;
diff --git a/Engine/Plugins/Experimental/ChaosMover/Source/ChaosMover/Public/ChaosMover/Character/Transitions/ChaosCharacterLandingCheck.h b/Engine/Plugins/Experimental/ChaosMover/Source/ChaosMover/Public/ChaosMover/Character/Transitions/ChaosCharacterLandingCheck.h
--- a/Engine/Plugins/Experimental/ChaosMover/Source/ChaosMover/Public/ChaosMover/Character/Transitions/ChaosCharacterLandingCheck.h
+++ b/Engine/Plugins/Experimental/ChaosMover/Source/ChaosMover/Public/ChaosMover/Character/Transitions/ChaosCharacterLandingCheck.h
@@ -18,10 +18,6 @@ public:
 	virtual FTransitionEvalResult Evaluate_Implementation(const FSimulationTickParams& Params) const override;
 	virtual void Trigger_Implementation(const FSimulationTickParams& Params) override;
 
-	/** Depth at which the pawn starts swimming */
-	UPROPERTY(BlueprintReadWrite, EditAnywhere, Category = Parameters, meta = (Units = "cm"))
-	float SwimmingStartImmersionDepth = 64.5f;
-
 	/** Height at which we consider the character to be on the ground */
 	UPROPERTY(BlueprintReadWrite, EditAnywhere, Category = Parameters, meta = (Units = "cm"))
 	float FloorDistanceTolerance = 0.5f;
