[ChaosMover] Misc. changes to support external usage of ChaosMover and Crouching with networking

diff --git a/Engine/Plugins/Experimental/ChaosMover/Source/ChaosMover/Private/Character/ChaosCharacterInputs.cpp b/Engine/Plugins/Experimental/ChaosMover/Source/ChaosMover/Private/Character/ChaosCharacterInputs.cpp
--- a/Engine/Plugins/Experimental/ChaosMover/Source/ChaosMover/Private/Character/ChaosCharacterInputs.cpp
+++ b/Engine/Plugins/Experimental/ChaosMover/Source/ChaosMover/Private/Character/ChaosCharacterInputs.cpp
@@ -81,6 +81,8 @@ void FChaosMoverCrouchInputs::Interpolate(const FMoverDataStructBase& From, cons
 
 void FChaosMoverCrouchInputs::Merge(const FMoverDataStructBase& From)
 {
+	const FChaosMoverCrouchInputs& TypedFrom = static_cast<const FChaosMoverCrouchInputs&>(From);
+	bWantsToCrouch |= TypedFrom.bWantsToCrouch;
 }
 
 bool FChaosMovementSettingsOverrides::NetSerialize(FArchive& Ar, UPackageMap* Map, bool& bOutSuccess)
diff --git a/Engine/Plugins/Experimental/ChaosMover/Source/ChaosMover/Public/ChaosMover/Character/ChaosCharacterMoverComponent.h b/Engine/Plugins/Experimental/ChaosMover/Source/ChaosMover/Public/ChaosMover/Character/ChaosCharacterMoverComponent.h
--- a/Engine/Plugins/Experimental/ChaosMover/Source/ChaosMover/Public/ChaosMover/Character/ChaosCharacterMoverComponent.h
+++ b/Engine/Plugins/Experimental/ChaosMover/Source/ChaosMover/Public/ChaosMover/Character/ChaosCharacterMoverComponent.h
@@ -39,7 +39,7 @@ public:
 	// Override the movement mode settings on a mode
 	// If the name is not set it will apply to the current mode
 	UFUNCTION(BlueprintCallable, Category = "Chaos Mover")
-	void OverrideMovementSettings(const FChaosMovementSettingsOverrides Overrides);
+	CHAOSMOVER_API void OverrideMovementSettings(const FChaosMovementSettingsOverrides Overrides);
 
 	// Cancel overrides of movement mode settings
 	// If the name is not set it will apply to the current mode
