ChaosMover - Change default for GT update to not be parallelized

#rnx

diff --git a/Engine/Plugins/Experimental/ChaosMover/Source/ChaosMover/Private/ChaosMoverDeveloperSettings.cpp b/Engine/Plugins/Experimental/ChaosMover/Source/ChaosMover/Private/ChaosMoverDeveloperSettings.cpp
--- a/Engine/Plugins/Experimental/ChaosMover/Source/ChaosMover/Private/ChaosMoverDeveloperSettings.cpp
+++ b/Engine/Plugins/Experimental/ChaosMover/Source/ChaosMover/Private/ChaosMoverDeveloperSettings.cpp
@@ -8,7 +8,7 @@ namespace UE::ChaosMover
 {
 	namespace CVars
 	{
-		bool bForceSingleThreadedGT = false;
+		bool bForceSingleThreadedGT = true;
 		FAutoConsoleVariableRef CVarChaosMoverForceSingleThreadedGT(TEXT("ChaosMover.ForceSingleThreadedGT"),
 			bForceSingleThreadedGT, TEXT("Force updates on the game thread to be single threaded."));
 
