ChaosMover - Remove default jumping transition from swimming and walking modes to easier allow users to use their own

#rnx

diff --git a/Engine/Plugins/Experimental/ChaosMover/Source/ChaosMover/Private/Character/Modes/ChaosSwimmingMode.cpp b/Engine/Plugins/Experimental/ChaosMover/Source/ChaosMover/Private/Character/Modes/ChaosSwimmingMode.cpp
--- a/Engine/Plugins/Experimental/ChaosMover/Source/ChaosMover/Private/Character/Modes/ChaosSwimmingMode.cpp
+++ b/Engine/Plugins/Experimental/ChaosMover/Source/ChaosMover/Private/Character/Modes/ChaosSwimmingMode.cpp
@@ -25,7 +25,6 @@ UChaosSwimmingMode::UChaosSwimmingMode(const FObjectInitializer& ObjectInitializ
 	TwistTorqueLimit = 0.0f;
 
 	Transitions.Add(CreateDefaultSubobject<UChaosCharacterWaterCheck>(TEXT("DefaultWaterCheck")));
-	Transitions.Add(CreateDefaultSubobject<UChaosCharacterJumpCheck>(TEXT("DefaultJumpCheck")));
 }
 
 void UChaosSwimmingMode::GenerateMove_Implementation(const FMoverTickStartData& StartState, const FMoverTimeStep& TimeStep, FProposedMove& OutProposedMove) const
diff --git a/Engine/Plugins/Experimental/ChaosMover/Source/ChaosMover/Private/Character/Modes/ChaosWalkingMode.cpp b/Engine/Plugins/Experimental/ChaosMover/Source/ChaosMover/Private/Character/Modes/ChaosWalkingMode.cpp
--- a/Engine/Plugins/Experimental/ChaosMover/Source/ChaosMover/Private/Character/Modes/ChaosWalkingMode.cpp
+++ b/Engine/Plugins/Experimental/ChaosMover/Source/ChaosMover/Private/Character/Modes/ChaosWalkingMode.cpp
@@ -35,7 +35,6 @@ UChaosWalkingMode::UChaosWalkingMode(const FObjectInitializer& ObjectInitializer
 	TwistTorqueLimit = 1500.0f;
 
 	Transitions.Add(CreateDefaultSubobject<UChaosCharacterFallingCheck>(TEXT("DefaultFallingCheck")));
-	Transitions.Add(CreateDefaultSubobject<UChaosCharacterJumpCheck>(TEXT("DefaultJumpCheck")));
 }
 
 void UChaosWalkingMode::UpdateConstraintSettings(Chaos::FCharacterGroundConstraintSettings& ConstraintSettings) const
