ChaosMover - Add export for simulation class accessor

#rnx

diff --git a/Engine/Plugins/Experimental/ChaosMover/Source/ChaosMover/Public/ChaosMover/Backends/ChaosMoverBackend.h b/Engine/Plugins/Experimental/ChaosMover/Source/ChaosMover/Public/ChaosMover/Backends/ChaosMoverBackend.h
--- a/Engine/Plugins/Experimental/ChaosMover/Source/ChaosMover/Public/ChaosMover/Backends/ChaosMoverBackend.h
+++ b/Engine/Plugins/Experimental/ChaosMover/Source/ChaosMover/Public/ChaosMover/Backends/ChaosMoverBackend.h
@@ -47,8 +47,8 @@ public:
 	CHAOSMOVER_API virtual void ConsumeOutputData(const FMoverTimeStep& TimeStep, const UE::ChaosMover::FSimulationOutputData& OutputData);
 	CHAOSMOVER_API virtual void FinalizeFrame(double ResultsTimeInMs);
 
-	UChaosMoverSimulation* GetSimulation();
-	const UChaosMoverSimulation* GetSimulation() const;
+	CHAOSMOVER_API UChaosMoverSimulation* GetSimulation();
+	CHAOSMOVER_API const UChaosMoverSimulation* GetSimulation() const;
 
 protected:
 	CHAOSMOVER_API virtual void InitSimulation();
