Implementation of FChaosCharacterApplyVelocityEffect::ToSimpleString

diff --git a/Engine/Plugins/Experimental/ChaosMover/Source/ChaosMover/Private/Character/Effects/ChaosCharacterApplyVelocityEffect.cpp b/Engine/Plugins/Experimental/ChaosMover/Source/ChaosMover/Private/Character/Effects/ChaosCharacterApplyVelocityEffect.cpp
--- a/Engine/Plugins/Experimental/ChaosMover/Source/ChaosMover/Private/Character/Effects/ChaosCharacterApplyVelocityEffect.cpp
+++ b/Engine/Plugins/Experimental/ChaosMover/Source/ChaosMover/Private/Character/Effects/ChaosCharacterApplyVelocityEffect.cpp
@@ -87,7 +87,7 @@ UScriptStruct* FChaosCharacterApplyVelocityEffect::GetScriptStruct() const
 
 FString FChaosCharacterApplyVelocityEffect::ToSimpleString() const
 {
-	return FString::Printf(TEXT("ChaosCharacterApplyVelocityEffect"));
+	return FString::Printf(TEXT("ChaosCharacterApplyVelocityEffect VelocityOrImpulseToApply = %s, Mode = %s"), *VelocityOrImpulseToApply.ToString(), *StaticEnum<EChaosMoverVelocityEffectMode>()->GetValueAsString(Mode));
 }
 
 void FChaosCharacterApplyVelocityEffect::AddReferencedObjects(FReferenceCollector& Collector)
