Build fix: Forgot to export some new functions (UChaosMoverSimulation::CanTeleport and Teleport)

diff --git a/Engine/Plugins/Experimental/ChaosMover/Source/ChaosMover/Public/ChaosMover/ChaosMoverSimulation.h b/Engine/Plugins/Experimental/ChaosMover/Source/ChaosMover/Public/ChaosMover/ChaosMoverSimulation.h
--- a/Engine/Plugins/Experimental/ChaosMover/Source/ChaosMover/Public/ChaosMover/ChaosMoverSimulation.h
+++ b/Engine/Plugins/Experimental/ChaosMover/Source/ChaosMover/Public/ChaosMover/ChaosMoverSimulation.h
@@ -150,9 +150,9 @@ public:
 	CHAOSMOVER_API virtual void AttemptTeleport(const FMoverTimeStep& TimeStep, const FTransform& TargetTransform, bool bUseActorRotation, FMoverSyncState& OutputState) override;
 protected:
 	// Whether the teleport can happen (does the movement mode allow it? Do we fit at the target transform?)
-	virtual bool CanTeleport(const FTransform& TargetTransform, bool bUseActorRotation, const FMoverSyncState& SyncState);
+	CHAOSMOVER_API virtual bool CanTeleport(const FTransform& TargetTransform, bool bUseActorRotation, const FMoverSyncState& SyncState);
 	// Actually perform the teleport
-	virtual void Teleport(const FTransform& TargetTransform, const FMoverTimeStep& TimeStep, FMoverSyncState& OutputState);
+	CHAOSMOVER_API virtual void Teleport(const FTransform& TargetTransform, const FMoverTimeStep& TimeStep, FMoverSyncState& OutputState);
 
 	// Object State.
 	// In kinematic mode, the updated primitive component needs both its SimulatePhysics and UpdateKinematicFromSimulation options checked
