ChaosMover - Adjust jump transition so that vertical jump velocity is always the same relative to the movement base regardless of the initial velocity

diff --git a/Engine/Plugins/Experimental/ChaosMover/Source/ChaosMover/Private/Character/Transitions/ChaosCharacterJumpCheck.cpp b/Engine/Plugins/Experimental/ChaosMover/Source/ChaosMover/Private/Character/Transitions/ChaosCharacterJumpCheck.cpp
--- a/Engine/Plugins/Experimental/ChaosMover/Source/ChaosMover/Private/Character/Transitions/ChaosCharacterJumpCheck.cpp
+++ b/Engine/Plugins/Experimental/ChaosMover/Source/ChaosMover/Private/Character/Transitions/ChaosCharacterJumpCheck.cpp
@@ -8,6 +8,7 @@
 #include "ChaosMover/ChaosMoverSimulation.h"
 #include "ChaosMover/ChaosMoverSimulationTypes.h"
 #include "ChaosMover/Character/Effects/ChaosCharacterApplyVelocityEffect.h"
+#include "ChaosMover/Utilities/ChaosGroundMovementUtils.h"
 #include "MoveLibrary/FloorQueryUtils.h"
 
 #include UE_INLINE_GENERATED_CPP_BY_NAME(ChaosCharacterJumpCheck)
@@ -57,6 +58,29 @@ void UChaosCharacterJumpCheck::Trigger_Implementation(const FSimulationTickParam
 	{
 		FVector JumpVelocity = JumpUpwardsSpeed * SimInputs->UpDir;
 
+		// Substract any current upwards velocity so that the jump velocity is
+		// always the same in the up direction
+		const FMoverDefaultSyncState* StartingSyncState = Params.StartState.SyncState.SyncStateCollection.FindDataByType<FMoverDefaultSyncState>();
+		if (StartingSyncState)
+		{
+			JumpVelocity -= StartingSyncState->GetVelocity_WorldSpace().Dot(SimInputs->UpDir) * SimInputs->UpDir;
+		}
+
+		// Add on the local ground velocity to make the jump velocity relative to the ground
+		FFloorCheckResult FloorResult;
+		const UMoverBlackboard* Blackboard = Simulation->GetBlackboard();
+		if (Blackboard && Blackboard->TryGet(CommonBlackboard::LastFloorResult, FloorResult))
+		{
+			if (Chaos::FPhysicsObjectHandle GroundObject = FloorResult.HitResult.PhysicsObject)
+			{
+				if (StartingSyncState)
+				{
+					const FVector LocalGroundVelocity = UChaosGroundMovementUtils::ComputeLocalGroundVelocity_Internal(StartingSyncState->GetLocation_WorldSpace(), FloorResult);
+					JumpVelocity += LocalGroundVelocity.Dot(SimInputs->UpDir) * SimInputs->UpDir;
+				}
+			}
+		}
+
 		TSharedPtr<FChaosCharacterApplyVelocityEffect> JumpMove = MakeShared<FChaosCharacterApplyVelocityEffect>();
 		JumpMove->VelocityOrImpulseToApply = JumpVelocity;
 		JumpMove->Mode = EChaosMoverVelocityEffectMode::AdditiveVelocity;
@@ -66,19 +90,14 @@ void UChaosCharacterJumpCheck::Trigger_Implementation(const FSimulationTickParam
 		// Apply an equal and opposite impulse to the ground
 		if (FractionalGroundReactionImpulse > 0.0f)
 		{
-			FFloorCheckResult FloorResult;
-			const UMoverBlackboard* Blackboard = Simulation->GetBlackboard();
-			if (Blackboard && Blackboard->TryGet(CommonBlackboard::LastFloorResult, FloorResult))
+			if (Chaos::FPhysicsObjectHandle GroundObject = FloorResult.HitResult.PhysicsObject)
 			{
-				if (Chaos::FPhysicsObjectHandle GroundObject = FloorResult.HitResult.PhysicsObject)
+				Chaos::FWritePhysicsObjectInterface_Internal Interface = Chaos::FPhysicsObjectInternalInterface::GetWrite();
+				const float CharacterMass = Interface.GetMass({ SimInputs->PhysicsObject });
+				FVector ImpulseToApply = -FractionalGroundReactionImpulse * CharacterMass * JumpVelocity;
+				if (Chaos::FPBDRigidParticleHandle* GroundParticle = Interface.GetRigidParticle(GroundObject))
 				{
-					Chaos::FWritePhysicsObjectInterface_Internal Interface = Chaos::FPhysicsObjectInternalInterface::GetWrite();
-					const float CharacterMass = Interface.GetMass({ SimInputs->PhysicsObject });
-					FVector ImpulseToApply = -FractionalGroundReactionImpulse * CharacterMass * JumpVelocity;
-					if (Chaos::FPBDRigidParticleHandle* GroundParticle = Interface.GetRigidParticle(GroundObject))
-					{
-						Private::ApplyImpulse(GroundParticle, ImpulseToApply, FloorResult.HitResult.ImpactPoint);
-					}
+					Private::ApplyImpulse(GroundParticle, ImpulseToApply, FloorResult.HitResult.ImpactPoint);
 				}
 			}
 		}
