ChaosMover - Fix computation of mas walkable slope cosine

#rnx

diff --git a/Engine/Plugins/Experimental/ChaosMover/Source/ChaosMover/Public/ChaosMover/Character/Settings/SharedChaosCharacterMovementSettings.h b/Engine/Plugins/Experimental/ChaosMover/Source/ChaosMover/Public/ChaosMover/Character/Settings/SharedChaosCharacterMovementSettings.h
--- a/Engine/Plugins/Experimental/ChaosMover/Source/ChaosMover/Public/ChaosMover/Character/Settings/SharedChaosCharacterMovementSettings.h
+++ b/Engine/Plugins/Experimental/ChaosMover/Source/ChaosMover/Public/ChaosMover/Character/Settings/SharedChaosCharacterMovementSettings.h
@@ -22,7 +22,7 @@ class USharedChaosCharacterMovementSettings : public UObject, public IMovementSe
 public:
 	float GetMaxWalkableSlopeCosine() const
 	{
-		return MaxWalkSlopeCosine;
+		return FMath::Cos(FMath::DegreesToRadians(MaxWalkableSlopeAngle));
 	}
 
 	/** Default max linear rate of deceleration when there is no controlled input */
@@ -96,16 +96,7 @@ public:
 	UPROPERTY(Category = "General", EditAnywhere, BlueprintReadWrite)
 	FName DefaultFallingMode = DefaultModeNames::Falling;
 
-protected:
-	/** Walkable slope angle, represented as cosine(max slope angle) for performance reasons. Ex: for max slope angle of 30 degrees, value is cosine(30 deg) = 0.866 */
-	UPROPERTY(BlueprintReadWrite, EditAnywhere, Category = "Ground Movement", BlueprintSetter=SetMaxWalkableSlopeAngle, meta = (ClampMin = "0", UIMin = "0", ClampMax = "90", UIMax = "90", ForceUnits = "degrees"))
+	// Maximum angle of slope that the character can walk on without sliding
+	UPROPERTY(BlueprintReadWrite, EditAnywhere, Category = "Ground Movement", meta = (ClampMin = "0", UIMin = "0", ClampMax = "90", UIMax = "90", ForceUnits = "degrees"))
 	float MaxWalkableSlopeAngle = 45.0f;
-
-	UFUNCTION(BlueprintSetter)
-	void SetMaxWalkableSlopeAngle(float Value)
-	{
-		MaxWalkSlopeCosine = FMath::Cos(FMath::DegreesToRadians(MaxWalkableSlopeAngle));
-	}
-
-	float MaxWalkSlopeCosine = 0.71f;
 };
