ChaosMover - Improvements to Walking Mode:
- Add the option to maintain horizontal velocity when going up and down slopes
- Remove velocity to target. Let the constraint move the character instead
- Always project the constraint target onto the mmovement plane

diff --git a/Engine/Plugins/Experimental/ChaosMover/Source/ChaosMover/Private/Character/Modes/ChaosWalkingMode.cpp b/Engine/Plugins/Experimental/ChaosMover/Source/ChaosMover/Private/Character/Modes/ChaosWalkingMode.cpp
--- a/Engine/Plugins/Experimental/ChaosMover/Source/ChaosMover/Private/Character/Modes/ChaosWalkingMode.cpp
+++ b/Engine/Plugins/Experimental/ChaosMover/Source/ChaosMover/Private/Character/Modes/ChaosWalkingMode.cpp
@@ -71,15 +71,11 @@ void UChaosWalkingMode::GenerateMove_Implementation(const FMoverTickStartData& S
 	// Try to use the floor as the basis for the intended move direction (i.e. try to walk along slopes, rather than into them)
 	UMoverBlackboard* SimBlackboard = Simulation->GetBlackboard_Mutable();
 	FFloorCheckResult LastFloorResult;
-	FVector MovementNormal;
-	if (SimBlackboard && SimBlackboard->TryGet(CommonBlackboard::LastFloorResult, LastFloorResult) && LastFloorResult.IsWalkableFloor())
+	FVector MovementNormal = UpDirection;
+	if (!bMaintainHorizontalGroundVelocity && SimBlackboard && SimBlackboard->TryGet(CommonBlackboard::LastFloorResult, LastFloorResult) && LastFloorResult.IsWalkableFloor())
 	{
 		MovementNormal = LastFloorResult.HitResult.ImpactNormal;
 	}
-	else
-	{
-		MovementNormal = UpDirection;
-	}
 
 	FRotator IntendedOrientation_WorldSpace;
 	// If there's no intent from input to change orientation, use the current orientation
@@ -151,6 +147,12 @@ void UChaosWalkingMode::GenerateMove_Implementation(const FMoverTickStartData& S
 
 		SimBlackboard->Set(CommonBlackboard::LastFloorResult, FloorResult);
 		SimBlackboard->Set(CommonBlackboard::LastWaterResult, WaterResult);
+
+		if (bMaintainHorizontalGroundVelocity)
+		{
+			// So far have assumed we are on level ground, so now add velocity up the slope
+			OutProposedMove.LinearVelocity -= UpDirection * OutProposedMove.LinearVelocity.Dot(FloorResult.HitResult.ImpactNormal) / UpDirection.Dot(FloorResult.HitResult.ImpactNormal);
+		}
 	};
 }
 
@@ -180,13 +182,13 @@ void UChaosWalkingMode::SimulationTick_Implementation(const FSimulationTickParam
 	const float DeltaSeconds = Params.TimeStep.StepMs * 0.001f;
 	const FVector UpDirection = DefaultSimInputs->UpDir;
 
-	FVector StartingGroundNormal = UpDirection;
+	FVector GroundNormal = UpDirection;
 	const UMoverBlackboard* SimBlackboard = Simulation->GetBlackboard();
 	FFloorCheckResult FloorResult;
 	if (SimBlackboard)
 	{
 		SimBlackboard->TryGet(CommonBlackboard::LastFloorResult, FloorResult);
-		StartingGroundNormal = FloorResult.HitResult.ImpactNormal;
+		GroundNormal = FloorResult.HitResult.ImpactNormal;
 	};
 
 	if (FloorResult.IsWalkableFloor())
@@ -194,11 +196,7 @@ void UChaosWalkingMode::SimulationTick_Implementation(const FSimulationTickParam
 		const float InitialHeightAboveFloor = FloorResult.FloorDist - GetTargetHeight();
 
 		// Put the target position on the floor at the target height
-		FVector TargetPosition = StartingSyncState->GetLocation_WorldSpace();
-		if (FloorResult.FloorDist > GetTargetHeight() + UE_KINDA_SMALL_NUMBER)
-		{
-			TargetPosition -= UpDirection * InitialHeightAboveFloor;
-		}
+		FVector TargetPosition = StartingSyncState->GetLocation_WorldSpace() - UpDirection * InitialHeightAboveFloor;
 
 		// The base movement mode does not apply gravity in walking mode so apply here.
 		// Also remove the gravity that will be applied by the physics simulation.
@@ -216,20 +214,16 @@ void UChaosWalkingMode::SimulationTick_Implementation(const FSimulationTickParam
 
 		if (bUseProposedMove)
 		{
-			const FVector ProposedMovePlaneVelocity = ProposedMove.LinearVelocity - ProposedMove.LinearVelocity.ProjectOnToNormal(StartingGroundNormal);
+			const FVector ProposedMovePlaneVelocity = ProposedMove.LinearVelocity - ProposedMove.LinearVelocity.ProjectOnToNormal(GroundNormal);
 
 			// If there is velocity intent in the normal direction then use the velocity from the proposed move. Otherwise
 			// retain the previous vertical velocity
 			FVector ProposedNormalVelocity = ProposedMove.LinearVelocity - ProposedMovePlaneVelocity;
 			if (ProposedNormalVelocity.SizeSquared() > UE_KINDA_SMALL_NUMBER)
 			{
-				TargetVelocity = ProposedMove.LinearVelocity;
+				TargetVelocity += ProposedNormalVelocity - TargetVelocity.ProjectOnToNormal(GroundNormal);
 				bNormalVelocityIntent = true;
 			}
-			else
-			{
-				TargetVelocity = ProposedMovePlaneVelocity + TargetVelocity.ProjectOnToNormal(StartingGroundNormal);
-			}
 
 			TargetPosition += ProposedMovePlaneVelocity * DeltaSeconds;
 		}
@@ -245,8 +239,8 @@ void UChaosWalkingMode::SimulationTick_Implementation(const FSimulationTickParam
 		const bool bIsGroundMoving = ProjectedGroundVelocity.SizeSquared() > UE_KINDA_SMALL_NUMBER;
 		const FVector ProjectedRelativeVelocity = TargetVelocity - ProjectedGroundVelocity;
 		const float ProjectedRelativeNormalVelocity = FloorResult.HitResult.ImpactNormal.Dot(TargetVelocity - ProjectedGroundVelocity);
-		const float ProjectedRelativeVerticalVelocity = UpDirection.Dot(TargetVelocity - ProjectedGroundVelocity);
-		const float VerticalVelocityLimit = bNormalVelocityIntent ? 2.0f / DeltaSeconds : FMath::Abs(DefaultSimInputs->UpDir.Dot(DefaultSimInputs->Gravity) * DeltaSeconds);
+		const float ProjectedRelativeVerticalVelocity = GroundNormal.Dot(TargetVelocity - ProjectedGroundVelocity);
+		const float VerticalVelocityLimit = bNormalVelocityIntent ? 2.0f / DeltaSeconds : FMath::Abs(GroundNormal.Dot(DefaultSimInputs->Gravity) * DeltaSeconds);
 
 		bool bIsLiftingOffSurface = false;
 		if ((ProjectedRelativeNormalVelocity > VerticalVelocityLimit) && bIsGroundMoving && (ProjectedRelativeVerticalVelocity > VerticalVelocityLimit))
diff --git a/Engine/Plugins/Experimental/ChaosMover/Source/ChaosMover/Public/ChaosMover/Character/Modes/ChaosWalkingMode.h b/Engine/Plugins/Experimental/ChaosMover/Source/ChaosMover/Public/ChaosMover/Character/Modes/ChaosWalkingMode.h
--- a/Engine/Plugins/Experimental/ChaosMover/Source/ChaosMover/Public/ChaosMover/Character/Modes/ChaosWalkingMode.h
+++ b/Engine/Plugins/Experimental/ChaosMover/Source/ChaosMover/Public/ChaosMover/Character/Modes/ChaosWalkingMode.h
@@ -50,6 +50,13 @@ public:
 	UPROPERTY(EditAnywhere, Category = "Movement Settings", meta = (ClampMin = "0", UIMin = "0", ClampMax = "1", UIMax = "1"))
 	float FractionalDownwardVelocityToTarget = 1.0f;
 
+	/**
+	 * If true, walking movement always maintains horizontal velocity when moving up ramps, which causes movement up ramps to be faster parallel to the ramp surface.
+	 * If false, then walking movement maintains velocity magnitude parallel to the ramp surface.
+	 */
+	UPROPERTY(EditAnywhere, Category = "Movement Settings")
+	uint8 bMaintainHorizontalGroundVelocity : 1 = 0;
+
 	UE_API virtual void UpdateConstraintSettings(Chaos::FCharacterGroundConstraintSettings& ConstraintSettings) const override;
 
 	UE_API virtual void GenerateMove_Implementation(const FMoverTickStartData& StartState, const FMoverTimeStep& TimeStep, FProposedMove& OutProposedMove) const override;
