ChaosMover: Fix for a crash bug found in Pimlico when using Chaos Pathed Movement Component

diff --git a/Engine/Plugins/Experimental/ChaosMover/Source/ChaosMover/Private/PathedMovement/ChaosPathedMovementDebugDrawComponent.cpp b/Engine/Plugins/Experimental/ChaosMover/Source/ChaosMover/Private/PathedMovement/ChaosPathedMovementDebugDrawComponent.cpp
--- a/Engine/Plugins/Experimental/ChaosMover/Source/ChaosMover/Private/PathedMovement/ChaosPathedMovementDebugDrawComponent.cpp
+++ b/Engine/Plugins/Experimental/ChaosMover/Source/ChaosMover/Private/PathedMovement/ChaosPathedMovementDebugDrawComponent.cpp
@@ -100,130 +100,127 @@ void UChaosPathedMovementDebugDrawComponent::CreateRenderState_Concurrent(FRegis
 	DebugSpheres.Reset();
 	DebugStars.Reset();
 
-	const UMoverComponent* MoverComp = GetMoverComponent();
-	if (!MoverComp)
-	{
-		return;
-	}
-
-	const TArray<UActorComponent*> DebugDrawInterfaceComponents = GetOwner()->GetComponentsByInterface(UChaosPathedMovementDebugDrawInterface::StaticClass());
-	if (!DebugDrawInterfaceComponents.IsEmpty())
+	if (const UMoverComponent* MoverComp = GetMoverComponent())
 	{
-		DebugDrawInterfaceObject = DebugDrawInterfaceComponents[0];
-	}
+		const TArray<UActorComponent*> DebugDrawInterfaceComponents = GetOwner()->GetComponentsByInterface(UChaosPathedMovementDebugDrawInterface::StaticClass());
+		if (!DebugDrawInterfaceComponents.IsEmpty())
+		{
+			DebugDrawInterfaceObject = DebugDrawInterfaceComponents[0];
+		}
 	
-	const UChaosPathedMovementMode* PathedMovementMode = MoverComp->FindMode_Mutable<UChaosPathedMovementMode>(MoverComp->StartingMovementMode);
-	if (PathedMovementMode && !PathedMovementMode->PathPatterns.IsEmpty() && (PathedMovementMode->bDebugDrawAggregatePath || PathedMovementMode->bAllowPatternDebugDrawing))
-	{
-		// Collect valid patterns and count the number of total curves to draw to reserve enough space for DebugLines
-		TArray<const UChaosPathedMovementPatternBase*> ValidPatterns;
-		int32 NumCurvesToDraw = PathedMovementMode->bDebugDrawAggregatePath ? 1 : 0;
-		for (UChaosPathedMovementPatternBase* PathPattern : PathedMovementMode->PathPatterns)
+		const UChaosPathedMovementMode* PathedMovementMode = MoverComp->FindMode_Mutable<UChaosPathedMovementMode>(MoverComp->StartingMovementMode);
+		if (PathedMovementMode && !PathedMovementMode->PathPatterns.IsEmpty() && (PathedMovementMode->bDebugDrawAggregatePath || PathedMovementMode->bAllowPatternDebugDrawing))
 		{
-			if (PathPattern)
+			// Collect valid patterns and count the number of total curves to draw to reserve enough space for DebugLines
+			TArray<const UChaosPathedMovementPatternBase*> ValidPatterns;
+			int32 NumCurvesToDraw = PathedMovementMode->bDebugDrawAggregatePath ? 1 : 0;
+			for (UChaosPathedMovementPatternBase* PathPattern : PathedMovementMode->PathPatterns)
 			{
-				ValidPatterns.Emplace(PathPattern);
-				if (PathPattern->DebugDrawUsingStepSamples())
+				if (PathPattern)
 				{
-					++NumCurvesToDraw;
-				}
+					ValidPatterns.Emplace(PathPattern);
+					if (PathPattern->DebugDrawUsingStepSamples())
+					{
+						++NumCurvesToDraw;
+					}
 
-				if (PathedMovementMode->bAllowPatternDebugDrawing)
-				{
-					PathPattern->AppendDebugDrawElements(*this, DebugBoundsBuilder);
+					if (PathedMovementMode->bAllowPatternDebugDrawing)
+					{
+						PathPattern->AppendDebugDrawElements(*this, DebugBoundsBuilder);
+					}
 				}
 			}
-		}
 
-		// Reserve enough debug lines to not have to reallocate on addition
-		const float NumSteps = CVarChaosPathedMovementTotalDebugDrawSteps.GetValueOnGameThread();
-		DebugLines.Reserve(DebugLines.Num() + NumSteps * NumCurvesToDraw);
-
-		// This is an array of path progress samples that we will use if we draw the aggregate path
-		// The spacing of samples needs to adapt to path patterns or the sampling resolution my
-		// skip entire patterns by stepping too coarsely. As we sample patterns, we may add the same sample multiple
-		// times or out of order. This array will then be sorted and values that are too close to the previous ones
-		// will be skipped
-		TArray<float> PathProgressSamples;
-		if (PathedMovementMode->bDebugDrawAggregatePath)
-		{
-			// We preserve for a minimum but may have to reallocate depending on patterns
-			PathProgressSamples.Reserve(NumSteps);
-		}
-		
-		// This is the inferred path basis
-		const FTransform LocalToWorld = GetOwner()->GetTransform();
+			// Reserve enough debug lines to not have to reallocate on addition
+			const float NumSteps = CVarChaosPathedMovementTotalDebugDrawSteps.GetValueOnGameThread();
+			DebugLines.Reserve(DebugLines.Num() + NumSteps * NumCurvesToDraw);
 
-		// Establish a previous transform for each path pattern, and add a sample for the aggreagte path for the start of each pattern
-		TArray<FTransform> PreviousStepPatternTransforms;
-		PreviousStepPatternTransforms.SetNumUninitialized(ValidPatterns.Num());
-		for (int32 PatternIdx = 0; PatternIdx < ValidPatterns.Num(); ++PatternIdx)
-		{
-			const UChaosPathedMovementPatternBase& Pattern = *ValidPatterns[PatternIdx];
+			// This is an array of path progress samples that we will use if we draw the aggregate path
+			// The spacing of samples needs to adapt to path patterns or the sampling resolution my
+			// skip entire patterns by stepping too coarsely. As we sample patterns, we may add the same sample multiple
+			// times or out of order. This array will then be sorted and values that are too close to the previous ones
+			// will be skipped
+			TArray<float> PathProgressSamples;
 			if (PathedMovementMode->bDebugDrawAggregatePath)
 			{
-				// Add an aggregate path sample for the beginning of this path.
-				PathProgressSamples.Add(Pattern.ConvertPatternToPathProgress(Pattern.StartAtPathProgress));
+				// We preserve for a minimum but may have to reallocate depending on patterns
+				PathProgressSamples.Reserve(NumSteps);
 			}
-			const bool bDrawSampledPattern = Pattern.DebugDrawUsingStepSamples() && PathedMovementMode->bAllowPatternDebugDrawing;
-			if (bDrawSampledPattern)
-			{
-				const FTransform FirstPatternTransform = Chaos::FRigidTransform3::MultiplyNoScale(Pattern.CalcMaskedTargetRelativeTransform(0.0f, LocalToWorld), LocalToWorld);
-				PreviousStepPatternTransforms[PatternIdx] = FirstPatternTransform;
-				DebugBoundsBuilder += FirstPatternTransform.GetLocation();
-			}
-		}			
 		
-		// Now sample each pattern to draw them, and add that sample to the aggregate path sample if we draw it
-		const int32 NumStepsToDraw = CVarChaosPathedMovementDisplayedDebugDrawSteps.GetValueOnGameThread() > 0 ? CVarChaosPathedMovementDisplayedDebugDrawSteps.GetValueOnGameThread() : NumSteps;
-		for (int32 Step = 1; Step <= NumStepsToDraw; ++Step)
-		{
+			// This is the inferred path basis
+			const FTransform LocalToWorld = GetOwner()->GetTransform();
+
+			// Establish a previous transform for each path pattern, and add a sample for the aggreagte path for the start of each pattern
+			TArray<FTransform> PreviousStepPatternTransforms;
+			PreviousStepPatternTransforms.SetNumUninitialized(ValidPatterns.Num());
 			for (int32 PatternIdx = 0; PatternIdx < ValidPatterns.Num(); ++PatternIdx)
 			{
 				const UChaosPathedMovementPatternBase& Pattern = *ValidPatterns[PatternIdx];
-				const float PatternProgress = Step / NumSteps;
 				if (PathedMovementMode->bDebugDrawAggregatePath)
 				{
-					PathProgressSamples.Add(Pattern.ConvertPatternToPathProgress(PatternProgress));
+					// Add an aggregate path sample for the beginning of this path.
+					PathProgressSamples.Add(Pattern.ConvertPatternToPathProgress(Pattern.StartAtPathProgress));
 				}
 				const bool bDrawSampledPattern = Pattern.DebugDrawUsingStepSamples() && PathedMovementMode->bAllowPatternDebugDrawing;
 				if (bDrawSampledPattern)
 				{
-					const FTransform PatternRelativeTransform = Pattern.CalcMaskedTargetRelativeTransform(PatternProgress, PreviousStepPatternTransforms[PatternIdx]);
+					const FTransform FirstPatternTransform = Chaos::FRigidTransform3::MultiplyNoScale(Pattern.CalcMaskedTargetRelativeTransform(0.0f, LocalToWorld), LocalToWorld);
+					PreviousStepPatternTransforms[PatternIdx] = FirstPatternTransform;
+					DebugBoundsBuilder += FirstPatternTransform.GetLocation();
+				}
+			}			
+		
+			// Now sample each pattern to draw them, and add that sample to the aggregate path sample if we draw it
+			const int32 NumStepsToDraw = CVarChaosPathedMovementDisplayedDebugDrawSteps.GetValueOnGameThread() > 0 ? CVarChaosPathedMovementDisplayedDebugDrawSteps.GetValueOnGameThread() : NumSteps;
+			for (int32 Step = 1; Step <= NumStepsToDraw; ++Step)
+			{
+				for (int32 PatternIdx = 0; PatternIdx < ValidPatterns.Num(); ++PatternIdx)
+				{
+					const UChaosPathedMovementPatternBase& Pattern = *ValidPatterns[PatternIdx];
+					const float PatternProgress = Step / NumSteps;
+					if (PathedMovementMode->bDebugDrawAggregatePath)
+					{
+						PathProgressSamples.Add(Pattern.ConvertPatternToPathProgress(PatternProgress));
+					}
+					const bool bDrawSampledPattern = Pattern.DebugDrawUsingStepSamples() && PathedMovementMode->bAllowPatternDebugDrawing;
+					if (bDrawSampledPattern)
+					{
+						const FTransform PatternRelativeTransform = Pattern.CalcMaskedTargetRelativeTransform(PatternProgress, PreviousStepPatternTransforms[PatternIdx]);
 
-					const FTransform PatternTransform = Chaos::FRigidTransform3::MultiplyNoScale(PatternRelativeTransform, LocalToWorld);
-					DebugLines.Emplace(PreviousStepPatternTransforms[PatternIdx].GetLocation(), PatternTransform.GetLocation(), Pattern.PatternDebugDrawColor, 1.f);
+						const FTransform PatternTransform = Chaos::FRigidTransform3::MultiplyNoScale(PatternRelativeTransform, LocalToWorld);
+						DebugLines.Emplace(PreviousStepPatternTransforms[PatternIdx].GetLocation(), PatternTransform.GetLocation(), Pattern.PatternDebugDrawColor, 1.f);
 
-					PreviousStepPatternTransforms[PatternIdx] = PatternTransform;
-					DebugBoundsBuilder += PatternTransform.GetLocation();
+						PreviousStepPatternTransforms[PatternIdx] = PatternTransform;
+						DebugBoundsBuilder += PatternTransform.GetLocation();
+					}
 				}
 			}
-		}
 
-		// Aggregate Path
-		if (PathedMovementMode->bDebugDrawAggregatePath)
-		{
-			// Sort overall path progress samples, they may have been added out of order when iterating over the patterns
-			PathProgressSamples.Sort();
-
-			// Establish the first point on the aggregate path
-			FVector PreviousStepAggregateLocation = Chaos::FRigidTransform3::MultiplyNoScale(PathedMovementMode->CalcTargetRelativeTransform(0.0f), LocalToWorld).GetLocation();
-			DebugBoundsBuilder += PreviousStepAggregateLocation;
-			float PreviousPathProgressSample = 0.0f;
-			// Draw the rest of the aggregate path, at the overall progress samples deduced from the sampling of path patterns
-			// and skip samples that are too close to the previous one we drew
-			for (float PathProgressSample : PathProgressSamples)
+			// Aggregate Path
+			if (PathedMovementMode->bDebugDrawAggregatePath)
 			{
-				if (FMath::IsNearlyEqual(PathProgressSample, PreviousPathProgressSample)) // Should we look at the difference in transform location instead?
+				// Sort overall path progress samples, they may have been added out of order when iterating over the patterns
+				PathProgressSamples.Sort();
+
+				// Establish the first point on the aggregate path
+				FVector PreviousStepAggregateLocation = Chaos::FRigidTransform3::MultiplyNoScale(PathedMovementMode->CalcTargetRelativeTransform(0.0f), LocalToWorld).GetLocation();
+				DebugBoundsBuilder += PreviousStepAggregateLocation;
+				float PreviousPathProgressSample = 0.0f;
+				// Draw the rest of the aggregate path, at the overall progress samples deduced from the sampling of path patterns
+				// and skip samples that are too close to the previous one we drew
+				for (float PathProgressSample : PathProgressSamples)
 				{
-					continue;
+					if (FMath::IsNearlyEqual(PathProgressSample, PreviousPathProgressSample)) // Should we look at the difference in transform location instead?
+					{
+						continue;
+					}
+					const float OverallPathProgress = PathProgressSample;
+					const FVector StepLocation = Chaos::FRigidTransform3::MultiplyNoScale(PathedMovementMode->CalcTargetRelativeTransform(OverallPathProgress), LocalToWorld).GetLocation();
+					DebugLines.Emplace(PreviousStepAggregateLocation, StepLocation, PathedMovementMode->PathDebugDrawColor, 2.f);
+					PreviousStepAggregateLocation = StepLocation;
+					DebugBoundsBuilder += StepLocation;
+					PreviousPathProgressSample = PathProgressSample;
 				}
-				const float OverallPathProgress = PathProgressSample;
-				const FVector StepLocation = Chaos::FRigidTransform3::MultiplyNoScale(PathedMovementMode->CalcTargetRelativeTransform(OverallPathProgress), LocalToWorld).GetLocation();
-				DebugLines.Emplace(PreviousStepAggregateLocation, StepLocation, PathedMovementMode->PathDebugDrawColor, 2.f);
-				PreviousStepAggregateLocation = StepLocation;
-				DebugBoundsBuilder += StepLocation;
-				PreviousPathProgressSample = PathProgressSample;
 			}
 		}
 	}
