ChaosMover: Adding rotation to point movement path pattern. This allows to control rotation of actors directly by setting a rotator, along a path that no longer needs to have translation. In conjunction with masking, this provides a direct way to affect only the rotation directly while another pattern can control the translation.

diff --git a/Engine/Plugins/Experimental/ChaosMover/Source/ChaosMover/Private/PathedMovement/ChaosPointMovementPathPattern.cpp b/Engine/Plugins/Experimental/ChaosMover/Source/ChaosMover/Private/PathedMovement/ChaosPointMovementPathPattern.cpp
--- a/Engine/Plugins/Experimental/ChaosMover/Source/ChaosMover/Private/PathedMovement/ChaosPointMovementPathPattern.cpp
+++ b/Engine/Plugins/Experimental/ChaosMover/Source/ChaosMover/Private/PathedMovement/ChaosPointMovementPathPattern.cpp
@@ -82,9 +82,11 @@ FTransform UChaosPointMovementPathPattern::CalcUnmaskedTargetRelativeTransform(f
 		
 		// Point locations are calculated in world space, but path targets need to be provided as relative transforms
 		const FVector TargetWorldLoc = UKismetMathLibrary::VLerp(PrevPoint.WorldLoc, NextPoint.WorldLoc, Alpha);
+		const FQuat TargetWorldRot = UKismetMathLibrary::Quat_Slerp(PrevPoint.WorldRot, NextPoint.WorldRot, Alpha);
 		const FVector TargetRelativeLoc = PathBasis.InverseTransformPositionNoScale(TargetWorldLoc);
+		const FQuat TargetRelativeRot = PathBasis.InverseTransformRotation(FQuat(TargetWorldRot));
 
-		return FTransform(TargetRelativeLoc);
+		return FTransform(TargetRelativeRot, TargetRelativeLoc);
 	}
 	return FTransform::Identity;
 }
@@ -93,36 +95,100 @@ void UChaosPointMovementPathPattern::RefreshAssignedPointProgress(bool bForceRef
 {
 	Chaos::EnsureIsInGameThreadContext();
 
-	if (!bHasAssignedPointProgress || bForceRefresh)
+	if (PathPoints.Num() > 0 && (!bHasAssignedPointProgress || bForceRefresh))
 	{
-		TotalPathDistance = 0.f;
-		
+		TotalPathDistance = 0.f;		
 		const FTransform& PathBasis = Simulation ? Simulation->GetMovementBasisTransform() : FTransform::Identity;
 		FVector PrevPointLocation = PathBasis.GetLocation();
+		FQuat PrevPointRotation = PathBasis.GetRotation();
+		int NumCollapsedIntervals = 0;
 
 		// Run through and establish point locations and distance
-		for (const FChaosPointMovementPathPoint& PathPoint : PathPoints)
+		for (int PointIndex = 0; PointIndex < PathPoints.Num(); ++PointIndex)
 		{
-			if (PathPoint.Basis == EChaosPointMovementLocationBasis::PreviousPoint)
+			const FChaosPointMovementPathPoint& PathPoint = PathPoints[PointIndex];
+			switch (PathPoint.Basis)
 			{
-				PathPoint.WorldLoc = PrevPointLocation + PathPoint.Location;	
+			case EChaosPointMovementLocationBasis::PreviousPoint:
+				PathPoint.WorldLoc = PrevPointLocation + PathPoint.Location;
+				break;
+			case EChaosPointMovementLocationBasis::PathOrigin:
+				PathPoint.WorldLoc = PathBasis.TransformPositionNoScale(PathPoint.Location);
+				break;
+			default:
+				PathPoint.WorldLoc = PathPoint.Location;
+				break;
 			}
-			else if (PathPoint.Basis == EChaosPointMovementLocationBasis::PathOrigin)
+
+			switch (PathPoint.RotationBasis)
 			{
-				PathPoint.WorldLoc = PathBasis.TransformPositionNoScale(PathPoint.Location);
+			case EChaosPointMovementLocationBasis::PreviousPoint:
+				PathPoint.WorldRot = PrevPointRotation * FQuat(PathPoint.Rotation);
+				break;
+			case EChaosPointMovementLocationBasis::PathOrigin:
+				PathPoint.WorldRot = PathBasis.TransformRotation(FQuat(PathPoint.Rotation));
+				break;
+			default:
+				PathPoint.WorldRot = FQuat(PathPoint.Rotation);
+				break;
 			}
-			
-			const float DistanceFromPrevious = (PathPoint.WorldLoc - PrevPointLocation).Size();
-			TotalPathDistance += DistanceFromPrevious;
-			PathPoint.DistanceFromStart = TotalPathDistance;
+
+			if (PointIndex > 0)
+			{
+				const float DistanceFromPrevious = (PathPoint.WorldLoc - PrevPointLocation).Size();
+				if (DistanceFromPrevious < SMALL_NUMBER)
+				{
+					NumCollapsedIntervals++;
+				}
+				TotalPathDistance += DistanceFromPrevious;
+			}
+			PathPoint.PathDistanceFromStart = TotalPathDistance;
 
 			PrevPointLocation = PathPoint.WorldLoc;
+			PrevPointRotation = FQuat(PathPoint.WorldRot);
 		}
 
-		// Go through again and assign progress based on distance
-		for (const FChaosPointMovementPathPoint& PathPoint : PathPoints)
+		// Go through again and assign progress based on distance or number of points traveled if collapsed
+		// Since C intervals of the N total intervals are collapsed, we assign C/N amount of total progress on collapsed intervals and 1-C/N on non collapsed intervals
+		// Progress on each collapsed interval is 1/N ( = C/N / C)
+		// Progress on non collapsed interval is the ratio of how long that interval is (d_i) and the total length along the path (D), normalized by the total progress on non collapsed intervals
+		// If P(i) is the progress at the end of interval index i (i in [1,N]),
+		// then the progress at the end of a collapsed interval is:
+		//     P(i) = P(i-1) + 1/N
+		// and the progress at the end of a non collapsed interval is:
+		//     P(i) = P(i-1) + d(i)/D * (1-C/N)
+		// The sum is given by:
+		// Sum(P(i)) = C * 1/N + Sum(d(i)) / D * (1-C/N)
+		// Since D = Sum(d(i)):
+		// Sum(P(i)) = C * 1/N + D/D*(1-C/N) = C/N + 1*C/N = 1
+		PathPoints[0].Progress = 0.0f;
+		int NumIntervals = PathPoints.Num() - 1;
+		if (NumIntervals > 0)
 		{
-			PathPoint.Progress = UKismetMathLibrary::SafeDivide(PathPoint.DistanceFromStart, TotalPathDistance);
+			float PreviousProgress = 0.0f;
+			float PreviousPathDistance = 0.0f;
+			float TotalProgressOnCollapsedIntervals = FMath::Clamp(float(NumCollapsedIntervals) / float(NumIntervals), 0.0f, 1.0f);
+			float TotalProgressOnSeparatedIntervals = FMath::Clamp(1.0f - TotalProgressOnCollapsedIntervals, 0.0f, 1.0f);
+			float ProgressOnOneCollapsedInterval = (TotalProgressOnCollapsedIntervals > SMALL_NUMBER) ? (1.0f / float(NumIntervals)) : 0.0f;
+			for (int PointIndex = 1; PointIndex < PathPoints.Num(); ++PointIndex)
+			{
+				const FChaosPointMovementPathPoint& PathPoint = PathPoints[PointIndex];
+				float DistanceOnPathSincePrevious = PathPoint.PathDistanceFromStart - PreviousPathDistance;
+				if (DistanceOnPathSincePrevious < SMALL_NUMBER)
+				{
+					// This is the end point of a collapsed interval, we increase progress by ProgressOnOneCollapsedInterval
+					PathPoint.Progress = PreviousProgress + ProgressOnOneCollapsedInterval;
+				}
+				else
+				{
+					// This is the end point of a non collapsed interval, we increase progress by TotalProgressOnSeparatedIntervals * d(i)/D
+					PathPoint.Progress = PreviousProgress + UKismetMathLibrary::SafeDivide(DistanceOnPathSincePrevious, TotalPathDistance) * TotalProgressOnSeparatedIntervals;
+				}
+
+				PathPoint.Progress = FMath::Clamp(PathPoint.Progress, 0.0f, 1.0f);
+				PreviousProgress = PathPoint.Progress;
+				PreviousPathDistance = PathPoint.PathDistanceFromStart;
+			}
 		}
 		
 		bHasAssignedPointProgress = true;
diff --git a/Engine/Plugins/Experimental/ChaosMover/Source/ChaosMover/Public/ChaosMover/PathedMovement/ChaosPointMovementPathPattern.h b/Engine/Plugins/Experimental/ChaosMover/Source/ChaosMover/Public/ChaosMover/PathedMovement/ChaosPointMovementPathPattern.h
--- a/Engine/Plugins/Experimental/ChaosMover/Source/ChaosMover/Public/ChaosMover/PathedMovement/ChaosPointMovementPathPattern.h
+++ b/Engine/Plugins/Experimental/ChaosMover/Source/ChaosMover/Public/ChaosMover/PathedMovement/ChaosPointMovementPathPattern.h
@@ -34,9 +34,17 @@ public:
 	UPROPERTY(EditAnywhere, BlueprintReadOnly, Category = Location)
 	EChaosPointMovementLocationBasis Basis = EChaosPointMovementLocationBasis::PreviousPoint;
 	
-	/** The location to move toward, relative to the path origin (i.e. root component location by default) */
+	/** The location to move toward (i.e. root component location by default) */
 	UPROPERTY(EditAnywhere, BlueprintReadOnly, Category = Location, meta = (MakeEditWidget = true))
-	FVector Location = FVector::Zero();
+	FVector Location = FVector::ZeroVector;
+
+	/** Basis that the rotation at this point is relative to */
+	UPROPERTY(EditAnywhere, BlueprintReadOnly, Category = Location)
+	EChaosPointMovementLocationBasis RotationBasis = EChaosPointMovementLocationBasis::PreviousPoint;
+
+	/** The rotation to rotate toward */
+	UPROPERTY(EditAnywhere, BlueprintReadOnly, Category = Location, meta = (MakeEditWidget = true))
+	FRotator Rotation = FRotator::ZeroRotator;
 
 	/**
 	 * The progress along the path that this point corresponds to.
@@ -47,9 +55,10 @@ public:
 	
 	/** Distance from the start of the path to this point. Read-only. */
 	// UPROPERTY(VisibleAnywhere, BlueprintReadOnly, Transient, Category = Location)
-	mutable float DistanceFromStart = 0.f;
+	mutable float PathDistanceFromStart = 0.f;
 
-	mutable FVector WorldLoc = FVector::Zero();
+	mutable FVector WorldLoc = FVector::ZeroVector;
+	mutable FQuat WorldRot = FQuat::Identity;
 };
 
 /** Movement pattern that moves between explicitly defined points */
