ChaosMover: Fix to pathed movement debug draws. Chaging API of path patterns to returning world coordinates, provided a path basis transform. This works better when the path specifies positions in world coordinates, because it can simply ignore the passed basis transform. In cases where rotations are irrelevant it can also avoid a rotation multiplication and return the basis transform unmodified.

diff --git a/Engine/Plugins/Experimental/ChaosMover/Source/ChaosMover/Private/ChaosMoverSimulation.cpp b/Engine/Plugins/Experimental/ChaosMover/Source/ChaosMover/Private/ChaosMoverSimulation.cpp
--- a/Engine/Plugins/Experimental/ChaosMover/Source/ChaosMover/Private/ChaosMoverSimulation.cpp
+++ b/Engine/Plugins/Experimental/ChaosMover/Source/ChaosMover/Private/ChaosMoverSimulation.cpp
@@ -902,17 +902,15 @@ void UChaosMoverSimulation::OnMovementModeChanged(const FMovementModeChangedEven
 
 					// If the mode was also a pathed mode, let's teleport to the position on the path
 					// we're supposed to be at, using CurrentPathProgress
-					Chaos::FRigidTransform3 TargetLastFrame = Chaos::FRigidTransform3::Identity;
 					IChaosPathedMovementModeInterface* NextPathedMovementMode = Cast<IChaosPathedMovementModeInterface>(NextMovementActuationInterface);
 					if (NextPathedMovementMode)
 					{	
 						const float PathProgress = PresimPathedMovementState ? PresimPathedMovementState->LastChangePathProgress : 0.0f;
 						// We resume the path at the reset progress
 						// (if not 0, CurrentPathProgress was either not reset properly when the movement mode changed, or deliberately set to a different value)
-						TargetLastFrame = NextPathedMovementMode->CalcTargetRelativeTransform(PathProgress);
+						FTransform TargetLastFrame = NextPathedMovementMode->CalcTargetTransform(PathProgress, MovementBasisTransform);
+						TeleportActuationTarget(TargetLastFrame);
 					}
-					Chaos::FRigidTransform3 WorldTransform = TargetLastFrame * MovementBasisTransform;
-					TeleportActuationTarget(WorldTransform);
 				}
 				else
 				{
diff --git a/Engine/Plugins/Experimental/ChaosMover/Source/ChaosMover/Private/PathedMovement/ChaosEllipticalMovementPathPattern.cpp b/Engine/Plugins/Experimental/ChaosMover/Source/ChaosMover/Private/PathedMovement/ChaosEllipticalMovementPathPattern.cpp
--- a/Engine/Plugins/Experimental/ChaosMover/Source/ChaosMover/Private/PathedMovement/ChaosEllipticalMovementPathPattern.cpp
+++ b/Engine/Plugins/Experimental/ChaosMover/Source/ChaosMover/Private/PathedMovement/ChaosEllipticalMovementPathPattern.cpp
@@ -9,7 +9,7 @@ void UChaosEllipticalMovementPathPattern::InitializePattern(UChaosMoverSimulatio
 	Super::InitializePattern(InSimulation);
 }
 
-FTransform UChaosEllipticalMovementPathPattern::CalcUnmaskedTargetRelativeTransform(float PatternProgress, const FTransform& CurTargetTransform) const
+FTransform UChaosEllipticalMovementPathPattern::CalcUnmaskedTargetTransform(float PatternProgress, const FTransform& BasisTransform) const
 {
 	// Calc the ellipse point
 	const float Angle = PatternProgress * FMath::DegreesToRadians(UsableArcAngle);
@@ -23,5 +23,6 @@ FTransform UChaosEllipticalMovementPathPattern::CalcUnmaskedTargetRelativeTransf
 
 	// Rotate the ellipse from the XY plane to whatever 3D plane is desired
 	const FVector TargetLocation = EllipsePlaneRotation.RotateVector(FVector(TargetLocationXY.X, TargetLocationXY.Y, 0.f));
-	return FTransform(TargetLocation);
+
+	return FTransform(BasisTransform.GetRotation(), BasisTransform.TransformPositionNoScale(TargetLocation), BasisTransform.GetScale3D());
 }
diff --git a/Engine/Plugins/Experimental/ChaosMover/Source/ChaosMover/Private/PathedMovement/ChaosPathedMovementDebugDrawComponent.cpp b/Engine/Plugins/Experimental/ChaosMover/Source/ChaosMover/Private/PathedMovement/ChaosPathedMovementDebugDrawComponent.cpp
--- a/Engine/Plugins/Experimental/ChaosMover/Source/ChaosMover/Private/PathedMovement/ChaosPathedMovementDebugDrawComponent.cpp
+++ b/Engine/Plugins/Experimental/ChaosMover/Source/ChaosMover/Private/PathedMovement/ChaosPathedMovementDebugDrawComponent.cpp
@@ -150,7 +150,7 @@ void UChaosPathedMovementDebugDrawComponent::CreateRenderState_Concurrent(FRegis
 			// This is the inferred path basis
 			const FTransform LocalToWorld = GetOwner()->GetTransform();
 
-			// Establish a previous transform for each path pattern, and add a sample for the aggreagte path for the start of each pattern
+			// Establish a previous transform for each path pattern, and add a sample for the aggregate path for the start of each pattern
 			TArray<FTransform> PreviousStepPatternTransforms;
 			PreviousStepPatternTransforms.SetNumUninitialized(ValidPatterns.Num());
 			for (int32 PatternIdx = 0; PatternIdx < ValidPatterns.Num(); ++PatternIdx)
@@ -164,7 +164,7 @@ void UChaosPathedMovementDebugDrawComponent::CreateRenderState_Concurrent(FRegis
 				const bool bDrawSampledPattern = Pattern.DebugDrawUsingStepSamples() && PathedMovementMode->bAllowPatternDebugDrawing;
 				if (bDrawSampledPattern)
 				{
-					const FTransform FirstPatternTransform = Chaos::FRigidTransform3::MultiplyNoScale(Pattern.CalcMaskedTargetRelativeTransform(0.0f, LocalToWorld), LocalToWorld);
+					const FTransform FirstPatternTransform = Pattern.CalcMaskedTargetTransform(0.0f, LocalToWorld);
 					PreviousStepPatternTransforms[PatternIdx] = FirstPatternTransform;
 					DebugBoundsBuilder += FirstPatternTransform.GetLocation();
 				}
@@ -185,9 +185,7 @@ void UChaosPathedMovementDebugDrawComponent::CreateRenderState_Concurrent(FRegis
 					const bool bDrawSampledPattern = Pattern.DebugDrawUsingStepSamples() && PathedMovementMode->bAllowPatternDebugDrawing;
 					if (bDrawSampledPattern)
 					{
-						const FTransform PatternRelativeTransform = Pattern.CalcMaskedTargetRelativeTransform(PatternProgress, PreviousStepPatternTransforms[PatternIdx]);
-
-						const FTransform PatternTransform = Chaos::FRigidTransform3::MultiplyNoScale(PatternRelativeTransform, LocalToWorld);
+						const FTransform PatternTransform = Pattern.CalcMaskedTargetTransform(PatternProgress, LocalToWorld);
 						DebugLines.Emplace(PreviousStepPatternTransforms[PatternIdx].GetLocation(), PatternTransform.GetLocation(), Pattern.PatternDebugDrawColor, 1.f);
 
 						PreviousStepPatternTransforms[PatternIdx] = PatternTransform;
@@ -203,7 +201,7 @@ void UChaosPathedMovementDebugDrawComponent::CreateRenderState_Concurrent(FRegis
 				PathProgressSamples.Sort();
 
 				// Establish the first point on the aggregate path
-				FVector PreviousStepAggregateLocation = Chaos::FRigidTransform3::MultiplyNoScale(PathedMovementMode->CalcTargetRelativeTransform(0.0f), LocalToWorld).GetLocation();
+				FVector PreviousStepAggregateLocation = PathedMovementMode->CalcTargetTransform(0.0f, LocalToWorld).GetLocation();
 				DebugBoundsBuilder += PreviousStepAggregateLocation;
 				float PreviousPathProgressSample = 0.0f;
 				// Draw the rest of the aggregate path, at the overall progress samples deduced from the sampling of path patterns
@@ -215,7 +213,7 @@ void UChaosPathedMovementDebugDrawComponent::CreateRenderState_Concurrent(FRegis
 						continue;
 					}
 					const float OverallPathProgress = PathProgressSample;
-					const FVector StepLocation = Chaos::FRigidTransform3::MultiplyNoScale(PathedMovementMode->CalcTargetRelativeTransform(OverallPathProgress), LocalToWorld).GetLocation();
+					const FVector StepLocation = PathedMovementMode->CalcTargetTransform(OverallPathProgress, LocalToWorld).GetLocation();
 					DebugLines.Emplace(PreviousStepAggregateLocation, StepLocation, PathedMovementMode->PathDebugDrawColor, 2.f);
 					PreviousStepAggregateLocation = StepLocation;
 					DebugBoundsBuilder += StepLocation;
@@ -338,10 +336,8 @@ void UChaosPathedMovementDebugDrawComponent::UpdatePreviewMeshComp(bool bForce)
 				ProgressPreviewMeshComp->SetVisibility(true);
 				
 				float OverallPathProgress = IChaosPathedMovementDebugDrawInterface::Execute_GetPreviewMeshOverallPathProgress(DebugDrawInterfaceObject);
-				const FTransform PreviewMeshRelativeTransform = PathedMovementMode->CalcTargetRelativeTransform(OverallPathProgress);
-
 				// Calculate the location without scale, then tack on the source mesh's scale
-				FTransform PreviewMeshWorldTransform = Chaos::FRigidTransform3::MultiplyNoScale(PreviewMeshRelativeTransform, MeshRoot->GetComponentTransform());
+				FTransform PreviewMeshWorldTransform = PathedMovementMode->CalcTargetTransform(OverallPathProgress, FTransform(MeshRoot->GetComponentTransform().ToMatrixNoScale()));
 				PreviewMeshWorldTransform.SetScale3D(MeshRoot->GetComponentScale());
 				ProgressPreviewMeshComp->SetWorldTransform(PreviewMeshWorldTransform);
 
diff --git a/Engine/Plugins/Experimental/ChaosMover/Source/ChaosMover/Private/PathedMovement/ChaosPathedMovementMode.cpp b/Engine/Plugins/Experimental/ChaosMover/Source/ChaosMover/Private/PathedMovement/ChaosPathedMovementMode.cpp
--- a/Engine/Plugins/Experimental/ChaosMover/Source/ChaosMover/Private/PathedMovement/ChaosPathedMovementMode.cpp
+++ b/Engine/Plugins/Experimental/ChaosMover/Source/ChaosMover/Private/PathedMovement/ChaosPathedMovementMode.cpp
@@ -277,14 +277,14 @@ void UChaosPathedMovementMode::SimulationTick_Implementation(const FSimulationTi
 		OnPlaybackBounced(Params.TimeStep.BaseSimTimeMs);
 	}
 
-	const Chaos::FRigidTransform3 TargetRelativeTransform = CalcTargetRelativeTransform(NewPathProgress);
-
 	FMoverDefaultSyncState& OutputDefaultSyncState = OutputState.SyncState.SyncStateCollection.FindOrAddMutableDataByType<FMoverDefaultSyncState>();
+
+	const FTransform TargetWorldTransform = CalcTargetTransform(NewPathProgress, Simulation->GetMovementBasisTransform());
 	
-	const FTransform& MovementBasisTransform = Simulation->GetMovementBasisTransform();
-	FVector WorldTranslation(MovementBasisTransform.TransformPosition(TargetRelativeTransform.GetTranslation()));
-	FRotator WorldRotation(MovementBasisTransform.TransformRotation(TargetRelativeTransform.GetRotation()));
-	OutputDefaultSyncState.SetTransforms_WorldSpace( WorldTranslation, WorldRotation,
+	FVector WorldTranslation(TargetWorldTransform.GetTranslation());
+	FRotator WorldRotation(TargetWorldTransform.GetRotation());
+
+	OutputDefaultSyncState.SetTransforms_WorldSpace(WorldTranslation, WorldRotation,
 											  FVector::ZeroVector, // Fix this eventually to return the velocity tangential to the path, including due to a change in path basis since last update
 											  nullptr ); // no movement base
 
@@ -297,8 +297,6 @@ void UChaosPathedMovementMode::SimulationTick_Implementation(const FSimulationTi
 		const FMoverDefaultSyncState* PreviousDefaultSyncState = Params.StartState.SyncState.SyncStateCollection.FindDataByType<FMoverDefaultSyncState>();
 		DebugData.PreviousWorldTranslation = PreviousDefaultSyncState ? PreviousDefaultSyncState->GetTransform_WorldSpace().GetLocation() : FVector::ZeroVector;
 		DebugData.PreviousWorldRotation = PreviousDefaultSyncState ? FRotator(PreviousDefaultSyncState->GetTransform_WorldSpace().GetRotation()) : FRotator::ZeroRotator;
-		DebugData.NewRelativeTranslation = TargetRelativeTransform.GetTranslation();
-		DebugData.NewRelativeRotation = FRotator(TargetRelativeTransform.GetRotation());
 		DebugData.NewWorldTranslation = WorldTranslation;
 		DebugData.NewWorldRotation = WorldRotation;
 		// Record state for debugging
@@ -450,20 +448,33 @@ void UChaosPathedMovementMode::SetOneWayTripDuration_BeginPlayOnly(float NewDura
 	}
 }
 
-Chaos::FRigidTransform3 UChaosPathedMovementMode::CalcTargetRelativeTransform(float OverallPathProgress) const
+FTransform UChaosPathedMovementMode::CalcTargetTransform(float OverallPathProgress, const FTransform& BasisTransform) const
 {
+	// The scale of basis transform does not affect the path itself
+	// Rather, the scale returned by CalcTargetTransform affects the scale of whatever is following the path
+	FTransform BasisTransformNoScale = BasisTransform;
+	BasisTransformNoScale.SetScale3D(FVector(1,1,1));
+
 	FTransform TargetRelativeTransform = FTransform::Identity;
-	
+	FVector AccumulatedRelativeScale = FVector(1,1,1);
+
 	for (UChaosPathedMovementPatternBase* PathPattern : PathPatterns)
 	{
 		if (PathPattern)
 		{
-			const FTransform PatternTransform = PathPattern->CalcTargetRelativeTransform(OverallPathProgress, TargetRelativeTransform);
-        	TargetRelativeTransform.Accumulate(PatternTransform);
+			const FTransform PatternWorldTransform = PathPattern->CalcTargetTransform(OverallPathProgress, BasisTransformNoScale);
+			const FTransform PatternRelativeTransform = FTransform(
+				BasisTransformNoScale.InverseTransformRotation(PatternWorldTransform.GetRotation()),
+				BasisTransformNoScale.InverseTransformPositionNoScale(PatternWorldTransform.GetLocation()));
+			TargetRelativeTransform.Accumulate(PatternRelativeTransform);
+			AccumulatedRelativeScale *= PatternWorldTransform.GetScale3D();
 		}
 	}
 
-	return TargetRelativeTransform;
+	return FTransform(
+		BasisTransformNoScale.TransformRotation(TargetRelativeTransform.GetRotation()),
+		BasisTransformNoScale.TransformPositionNoScale(TargetRelativeTransform.GetLocation()),
+		BasisTransform.GetScale3D() * AccumulatedRelativeScale);
 }
 
 void UChaosPathedMovementMode::SetDefaultJointConstraintProperties()
diff --git a/Engine/Plugins/Experimental/ChaosMover/Source/ChaosMover/Private/PathedMovement/ChaosPathedMovementPatternBase.cpp b/Engine/Plugins/Experimental/ChaosMover/Source/ChaosMover/Private/PathedMovement/ChaosPathedMovementPatternBase.cpp
--- a/Engine/Plugins/Experimental/ChaosMover/Source/ChaosMover/Private/PathedMovement/ChaosPathedMovementPatternBase.cpp
+++ b/Engine/Plugins/Experimental/ChaosMover/Source/ChaosMover/Private/PathedMovement/ChaosPathedMovementPatternBase.cpp
@@ -124,20 +124,20 @@ float UChaosPathedMovementPatternBase::ConvertPatternToPathProgress(float FinalP
 	return FMath::Clamp(StartAtPathProgress + (EndAtPathProgress - StartAtPathProgress) * CurLoopPathProgress, 0.0f, 1.0f);
 }
 
-FTransform UChaosPathedMovementPatternBase::CalcTargetRelativeTransform(float OverallPathProgress, const FTransform& CurTargetTransform) const
+FTransform UChaosPathedMovementPatternBase::CalcTargetTransform(float OverallPathProgress, const FTransform& BasisTransform) const
 {
 	float PatternProgress = ConvertPathToPatternProgress(OverallPathProgress);
-	return CalcMaskedTargetRelativeTransform(PatternProgress, CurTargetTransform);
+	return CalcMaskedTargetTransform(PatternProgress, BasisTransform);
 }
 
-FTransform UChaosPathedMovementPatternBase::CalcMaskedTargetRelativeTransform(float PatternProgress, const FTransform& CurTargetTransform) const
+FTransform UChaosPathedMovementPatternBase::CalcMaskedTargetTransform(float PatternProgress, const FTransform& BasisTransform) const
 {
 	if (PatternProgress < 0.0f)
 	{
 		return FTransform();
 	}
 
-	const FTransform UnmaskedTarget = CalcUnmaskedTargetRelativeTransform(PatternProgress, CurTargetTransform);
+	const FTransform UnmaskedTarget = CalcUnmaskedTargetTransform(PatternProgress, BasisTransform);
 
 	if (bOrientComponentToPath)
 	{
diff --git a/Engine/Plugins/Experimental/ChaosMover/Source/ChaosMover/Private/PathedMovement/ChaosPointMovementPathPattern.cpp b/Engine/Plugins/Experimental/ChaosMover/Source/ChaosMover/Private/PathedMovement/ChaosPointMovementPathPattern.cpp
--- a/Engine/Plugins/Experimental/ChaosMover/Source/ChaosMover/Private/PathedMovement/ChaosPointMovementPathPattern.cpp
+++ b/Engine/Plugins/Experimental/ChaosMover/Source/ChaosMover/Private/PathedMovement/ChaosPointMovementPathPattern.cpp
@@ -8,6 +8,7 @@
 #include "Framework/Threading.h"
 #include "ChaosMover/PathedMovement/ChaosPathedMovementMode.h"
 #include "ChaosMover/PathedMovement/ChaosPathedMovementDebugDrawComponent.h"
+#include "MoverComponent.h"
 
 #include UE_INLINE_GENERATED_CPP_BY_NAME(ChaosPointMovementPathPattern)
 
@@ -15,7 +16,7 @@ void UChaosPointMovementPathPattern::InitializePattern(UChaosMoverSimulation* In
 {
 	Super::InitializePattern(InSimulation);
 
-	RefreshAssignedPointProgress(true);
+	RefreshAssignedPointProgress(GetPathBasisTransform());
 }
 
 #if WITH_EDITOR
@@ -23,174 +24,248 @@ void UChaosPointMovementPathPattern::PostEditChangeProperty(struct FPropertyChan
 {
 	Super::PostEditChangeProperty(PropertyChangedEvent);
 
-	RefreshAssignedPointProgress(true);
+	RefreshAssignedPointProgress(GetPathBasisTransform());
 }
 #endif
 
-void UChaosPointMovementPathPattern::AppendDebugDrawElements(UChaosPathedMovementDebugDrawComponent& DebugDrawComp, FBoxSphereBounds::Builder& InOutDebugBoundsBuilder)
+FTransform UChaosPointMovementPathPattern::GetPathBasisTransform() const
 {
-	if (!PathPoints.IsEmpty())
+	FTransform BasisTransform = FTransform::Identity;
+	if (GetSimulation())
 	{
-		const FTransform& PathOriginTransform = Simulation ? Simulation->GetMovementBasisTransform() : FTransform::Identity;
-		FVector PrevPointLocation = PathOriginTransform.GetLocation();
-		for (const FChaosPointMovementPathPoint& PathPoint : PathPoints)
+		BasisTransform = GetSimulation()->GetMovementBasisTransform();
+	}
+	else if (UChaosPathedMovementMode* ChaosPathedMovementMode = Cast<UChaosPathedMovementMode>(GetOuter()))
+	{
+		if (UActorComponent* ActorComponent = ChaosPathedMovementMode->GetMoverComponent())
 		{
-			DebugDrawComp.DebugSpheres.Emplace(4.f, PrevPointLocation, PatternDebugDrawColor);
-
-			FVector PointLocation = PathPoint.WorldLoc;
-			if (PathPoint.Basis == EChaosPointMovementLocationBasis::PreviousPoint)
+			if (ActorComponent->GetOwner())
 			{
-				PointLocation = PrevPointLocation + PathPoint.Location;	
+				BasisTransform = ActorComponent->GetOwner()->GetTransform();
 			}
-			else if (PathPoint.Basis == EChaosPointMovementLocationBasis::PathOrigin)
-			{
-				PointLocation = PathOriginTransform.TransformPositionNoScale(PathPoint.Location);
-			}
-			
-			DebugDrawComp.DebugLines.Emplace(PrevPointLocation, PointLocation, PatternDebugDrawColor, 1.f);
-			InOutDebugBoundsBuilder += PointLocation;
-			
-			PrevPointLocation = PointLocation;
 		}
+	}
+	return BasisTransform;
+}
+
+void UChaosPointMovementPathPattern::AppendDebugDrawElements(UChaosPathedMovementDebugDrawComponent& DebugDrawComp, FBoxSphereBounds::Builder& InOutDebugBoundsBuilder)
+{
+	if (!PathPoints.IsEmpty())
+	{
+		const FTransform& PathBasisTransform = DebugDrawComp.GetOwner()->GetTransform();
+
+		FVector PreviousPathPointLocation = PathPoints[0].GetWorldTransform(PathBasisTransform).GetLocation();
+		DebugDrawComp.DebugSpheres.Emplace(4.f, PreviousPathPointLocation, PatternDebugDrawColor);
+		InOutDebugBoundsBuilder += PreviousPathPointLocation;
+		
+		for (int PointIndex = 1; PointIndex < PathPoints.Num(); ++PointIndex)
+		{
+			const FChaosPointMovementPathPoint& PathPoint = PathPoints[PointIndex];
+			FVector PathPointLocation = PathPoint.GetWorldTransform(PathBasisTransform).GetLocation();
+			DebugDrawComp.DebugSpheres.Emplace(4.f, PathPointLocation, PatternDebugDrawColor);
 
-		DebugDrawComp.DebugSpheres.Emplace(4.f, PrevPointLocation, PatternDebugDrawColor);
+			DebugDrawComp.DebugLines.Emplace(PreviousPathPointLocation, PathPointLocation, PatternDebugDrawColor, 1.f);
+			PreviousPathPointLocation = PathPointLocation;
+
+			InOutDebugBoundsBuilder += PathPointLocation;
+		}
 	}
 }
 
-FTransform UChaosPointMovementPathPattern::CalcUnmaskedTargetRelativeTransform(float PatternProgress, const FTransform& CurTargetTransform) const
+FTransform FChaosPointMovementPathPoint::GetWorldTransform(const FTransform& BasisTransform) const
+{
+	return FTransform(
+		(EffectiveRotationBasis == EChaosPointMovementLocationBasis::World) ? Rotation : FRotator(BasisTransform.TransformRotation(FQuat(BasedRotation))),
+		(EffectiveBasis == EChaosPointMovementLocationBasis::World) ? Location : BasisTransform.TransformPositionNoScale(BasedLocation),
+		BasisTransform.GetScale3D()
+	);
+}
+
+FTransform UChaosPointMovementPathPattern::CalcUnmaskedTargetTransform(float PatternProgress, const FTransform& BasisTransform) const
 {
-	// The destination point is the first point at or equal to the target progress
-	int32 DestPointIdx = 0;
-	for (; DestPointIdx < PathPoints.Num() && PathPoints[DestPointIdx].Progress < PatternProgress; ++DestPointIdx) {}
 
-	if (PathPoints.IsValidIndex(DestPointIdx))
+	if (!PathPoints.IsEmpty())
 	{
-		const FTransform& PathBasis = Simulation ? Simulation->GetMovementBasisTransform() : FTransform::Identity;
-		const FChaosPointMovementPathPoint& NextPoint = PathPoints[DestPointIdx];
-		FChaosPointMovementPathPoint PrevPoint;
-		if (DestPointIdx > 0)
+		if (PatternProgress < SMALL_NUMBER)
 		{
-			PrevPoint = PathPoints[DestPointIdx - 1];
+			return PathPoints[0].GetWorldTransform(BasisTransform);
 		}
 		else
 		{
-			PrevPoint.WorldLoc = PathBasis.GetLocation();
-		}
+			// The destination point is the first point at or equal to the target progress
+			int32 DestPointIdx = 0;
+			for (; DestPointIdx < PathPoints.Num() && PathPoints[DestPointIdx].Progress < PatternProgress; ++DestPointIdx) {}
 
-		const float ProgressSinceLastPoint = PatternProgress - PrevPoint.Progress;
-		const float Alpha = (PatternProgress == PrevPoint.Progress) ? 0.0f : ProgressSinceLastPoint / (NextPoint.Progress - PrevPoint.Progress);
-		
-		// Point locations are calculated in world space, but path targets need to be provided as relative transforms
-		const FVector TargetWorldLoc = UKismetMathLibrary::VLerp(PrevPoint.WorldLoc, NextPoint.WorldLoc, Alpha);
-		const FQuat TargetWorldRot = UKismetMathLibrary::Quat_Slerp(PrevPoint.WorldRot, NextPoint.WorldRot, Alpha);
-		const FVector TargetRelativeLoc = PathBasis.InverseTransformPositionNoScale(TargetWorldLoc);
-		const FQuat TargetRelativeRot = PathBasis.InverseTransformRotation(FQuat(TargetWorldRot));
+			if (PathPoints.IsValidIndex(DestPointIdx))
+			{
+				const FChaosPointMovementPathPoint& NextPoint = PathPoints[DestPointIdx];
+				FChaosPointMovementPathPoint PrevPoint;
+				if (DestPointIdx > 0)
+				{
+					PrevPoint = PathPoints[DestPointIdx - 1];
+				}
+				else
+				{
+					PrevPoint.BasedLocation = BasisTransform.GetLocation();
+					PrevPoint.BasedRotation = BasisTransform.GetRotation();
+					PrevPoint.EffectiveBasis = EChaosPointMovementLocationBasis::World;
+					PrevPoint.EffectiveRotationBasis = EChaosPointMovementLocationBasis::World;
+				}
+
+				const float ProgressSinceLastPoint = PatternProgress - PrevPoint.Progress;
+				const float Alpha = (PatternProgress == PrevPoint.Progress) ? 0.0f : ProgressSinceLastPoint / (NextPoint.Progress - PrevPoint.Progress);
+				const FTransform PreviousWorldTransform = PrevPoint.GetWorldTransform(BasisTransform);
+				const FTransform NextWorldTransform = NextPoint.GetWorldTransform(BasisTransform);
+				const FVector TargetWorldLocation = UKismetMathLibrary::VLerp(PreviousWorldTransform.GetLocation(), NextWorldTransform.GetLocation(), Alpha);
+				const FQuat TargetWorldRotation = UKismetMathLibrary::Quat_Slerp(PreviousWorldTransform.GetRotation(), NextWorldTransform.GetRotation(), Alpha);
 
-		return FTransform(TargetRelativeRot, TargetRelativeLoc);
+				return FTransform(TargetWorldRotation, TargetWorldLocation, BasisTransform.GetScale3D());
+			}
+			else
+			{
+				return PathPoints[PathPoints.Num() - 1].GetWorldTransform(BasisTransform);
+			}
+		}
 	}
-	return FTransform::Identity;
+
+	return BasisTransform;
 }
 
-void UChaosPointMovementPathPattern::RefreshAssignedPointProgress(bool bForceRefresh) const
+void UChaosPointMovementPathPattern::RefreshAssignedPointProgress(const FTransform& PathBasisTransform) const
 {
 	Chaos::EnsureIsInGameThreadContext();
 
-	if (PathPoints.Num() > 0 && (!bHasAssignedPointProgress || bForceRefresh))
+	if (!bHasAssignedPointProgress || !CachedPathBasisTransform.Equals(PathBasisTransform, SMALL_NUMBER))
 	{
-		TotalPathDistance = 0.f;		
-		const FTransform& PathBasis = Simulation ? Simulation->GetMovementBasisTransform() : FTransform::Identity;
-		FVector PrevPointLocation = PathBasis.GetLocation();
-		FQuat PrevPointRotation = PathBasis.GetRotation();
-		int NumCollapsedIntervals = 0;
-
-		// Run through and establish point locations and distance
-		for (int PointIndex = 0; PointIndex < PathPoints.Num(); ++PointIndex)
+		if (PathPoints.Num() > 0)
 		{
-			const FChaosPointMovementPathPoint& PathPoint = PathPoints[PointIndex];
-			switch (PathPoint.Basis)
-			{
-			case EChaosPointMovementLocationBasis::PreviousPoint:
-				PathPoint.WorldLoc = PrevPointLocation + PathPoint.Location;
-				break;
-			case EChaosPointMovementLocationBasis::PathOrigin:
-				PathPoint.WorldLoc = PathBasis.TransformPositionNoScale(PathPoint.Location);
-				break;
-			default:
-				PathPoint.WorldLoc = PathPoint.Location;
-				break;
-			}
+			const FChaosPointMovementPathPoint& FirstPoint = PathPoints[0];
+			// For point 0, being relative to the previous point means being relative to the path basis
+			EChaosPointMovementLocationBasis LocationBasis = (FirstPoint.Basis != EChaosPointMovementLocationBasis::PreviousPoint) ? PathPoints[0].Basis : EChaosPointMovementLocationBasis::PathOrigin;
+			EChaosPointMovementLocationBasis RotationBasis = (FirstPoint.RotationBasis != EChaosPointMovementLocationBasis::PreviousPoint) ? PathPoints[0].RotationBasis : EChaosPointMovementLocationBasis::PathOrigin;
+			FirstPoint.EffectiveBasis = LocationBasis;
+			FirstPoint.EffectiveRotationBasis = RotationBasis;
+			FirstPoint.BasedLocation = FirstPoint.Location;
+			FirstPoint.BasedRotation = FQuat(FirstPoint.Rotation);
+			FirstPoint.PathDistanceFromStart = 0.0f;
 
-			switch (PathPoint.RotationBasis)
-			{
-			case EChaosPointMovementLocationBasis::PreviousPoint:
-				PathPoint.WorldRot = PrevPointRotation * FQuat(PathPoint.Rotation);
-				break;
-			case EChaosPointMovementLocationBasis::PathOrigin:
-				PathPoint.WorldRot = PathBasis.TransformRotation(FQuat(PathPoint.Rotation));
-				break;
-			default:
-				PathPoint.WorldRot = FQuat(PathPoint.Rotation);
-				break;
-			}
+			FTransform PreviousWorldPointTransform = FirstPoint.GetWorldTransform(PathBasisTransform);
+
+			TotalPathDistance = 0.f;
+			int NumCollapsedIntervals = 0;
 
-			if (PointIndex > 0)
+			// Run through and establish point locations, rotations and distance
+			for (int PointIndex = 1; PointIndex < PathPoints.Num(); ++PointIndex)
 			{
-				const float DistanceFromPrevious = (PathPoint.WorldLoc - PrevPointLocation).Size();
+				const FChaosPointMovementPathPoint& PathPoint = PathPoints[PointIndex];
+				switch (PathPoint.Basis)
+				{
+				case EChaosPointMovementLocationBasis::PreviousPoint:
+					// This point's transform is relative to the previous point's transform, whether that transform was relative to world or the path basis
+					// since this point will retain the previous point's effective basis
+					PathPoint.BasedLocation = PreviousWorldPointTransform.TransformPositionNoScale(PathPoint.Location);
+					if (LocationBasis == EChaosPointMovementLocationBasis::PathOrigin)
+					{
+						// If the effective location basis is path origin (remember that it cannot be PreviousPoint),
+						// we need to transform this world position back into a position relative to the path basis
+						PathPoint.BasedLocation = PathBasisTransform.InverseTransformPositionNoScale(PathPoint.BasedLocation);
+					}
+					
+					PathPoint.EffectiveBasis = LocationBasis;
+					// We avoid updating LocationBasis and we will not update it until we find a Basis that is not PreviousPoint
+					break;
+				default:
+					PathPoint.EffectiveBasis = PathPoint.Basis;
+					PathPoint.BasedLocation = PathPoint.Location;
+					// We update LocationBasis, since it is not PreviousPoint
+					LocationBasis = PathPoint.Basis;
+					break;
+				}
+
+				switch (PathPoint.RotationBasis)
+				{
+				case EChaosPointMovementLocationBasis::PreviousPoint:
+					PathPoint.BasedRotation = PreviousWorldPointTransform.GetRotation() * FQuat(PathPoint.Rotation);
+					if (RotationBasis == EChaosPointMovementLocationBasis::PathOrigin)
+					{
+						// If the effective rotation basis is path origin (remember that it cannot be PreviousPoint),
+						// we need to transform this world rotation back into a rotation relative to the path basis
+						PathPoint.BasedRotation = PathBasisTransform.InverseTransformRotation(PathPoint.BasedRotation);
+					}
+					PathPoint.EffectiveRotationBasis = RotationBasis;
+					// We avoid updating RotationBasis and we will not update it until we find a RotationBasis that is not PreviousPoint
+					break;
+				default:
+					PathPoint.EffectiveRotationBasis = PathPoint.RotationBasis;
+					PathPoint.BasedRotation = FQuat(PathPoint.Rotation);
+					// We update RotationBasis, since it is not PreviousPoint
+					RotationBasis = PathPoint.RotationBasis;
+					break;
+				}
+
+				const FTransform CurrentWorldPointTransform = PathPoint.GetWorldTransform(PathBasisTransform);
+
+
+				const float DistanceFromPrevious = (CurrentWorldPointTransform.GetLocation() - PreviousWorldPointTransform.GetLocation()).Size();
 				if (DistanceFromPrevious < SMALL_NUMBER)
 				{
 					NumCollapsedIntervals++;
 				}
 				TotalPathDistance += DistanceFromPrevious;
-			}
-			PathPoint.PathDistanceFromStart = TotalPathDistance;
 
-			PrevPointLocation = PathPoint.WorldLoc;
-			PrevPointRotation = FQuat(PathPoint.WorldRot);
-		}
+				PathPoint.PathDistanceFromStart = TotalPathDistance;
 
-		// Go through again and assign progress based on distance or number of points traveled if collapsed
-		// Since C intervals of the N total intervals are collapsed, we assign C/N amount of total progress on collapsed intervals and 1-C/N on non collapsed intervals
-		// Progress on each collapsed interval is 1/N ( = C/N / C)
-		// Progress on non collapsed interval is the ratio of how long that interval is (d_i) and the total length along the path (D), normalized by the total progress on non collapsed intervals
-		// If P(i) is the progress at the end of interval index i (i in [1,N]),
-		// then the progress at the end of a collapsed interval is:
-		//     P(i) = P(i-1) + 1/N
-		// and the progress at the end of a non collapsed interval is:
-		//     P(i) = P(i-1) + d(i)/D * (1-C/N)
-		// The sum is given by:
-		// Sum(P(i)) = C * 1/N + Sum(d(i)) / D * (1-C/N)
-		// Since D = Sum(d(i)):
-		// Sum(P(i)) = C * 1/N + D/D*(1-C/N) = C/N + 1*C/N = 1
-		PathPoints[0].Progress = 0.0f;
-		int NumIntervals = PathPoints.Num() - 1;
-		if (NumIntervals > 0)
-		{
-			float PreviousProgress = 0.0f;
-			float PreviousPathDistance = 0.0f;
-			float TotalProgressOnCollapsedIntervals = FMath::Clamp(float(NumCollapsedIntervals) / float(NumIntervals), 0.0f, 1.0f);
-			float TotalProgressOnSeparatedIntervals = FMath::Clamp(1.0f - TotalProgressOnCollapsedIntervals, 0.0f, 1.0f);
-			float ProgressOnOneCollapsedInterval = (TotalProgressOnCollapsedIntervals > SMALL_NUMBER) ? (1.0f / float(NumIntervals)) : 0.0f;
-			for (int PointIndex = 1; PointIndex < PathPoints.Num(); ++PointIndex)
+				PreviousWorldPointTransform = CurrentWorldPointTransform;
+			}
+
+			// Go through again and assign progress based on distance or number of points traveled if collapsed
+			// Since C intervals of the N total intervals are collapsed, we assign C/N amount of total progress on collapsed intervals and 1-C/N on non collapsed intervals
+			// Progress on each collapsed interval is 1/N ( = C/N / C)
+			// Progress on non collapsed interval is the ratio of how long that interval is (d_i) and the total length along the path (D), normalized by the total progress on non collapsed intervals
+			// If P(i) is the progress at the end of interval index i (i in [1,N]),
+			// then the progress at the end of a collapsed interval is:
+			//     P(i) = P(i-1) + 1/N
+			// and the progress at the end of a non collapsed interval is:
+			//     P(i) = P(i-1) + d(i)/D * (1-C/N)
+			// The sum is given by:
+			// Sum(P(i)) = C * 1/N + Sum(d(i)) / D * (1-C/N)
+			// Since D = Sum(d(i)):
+			// Sum(P(i)) = C * 1/N + D/D*(1-C/N) = C/N + 1*C/N = 1
+			PathPoints[0].Progress = 0.0f;
+			int NumIntervals = PathPoints.Num() - 1;
+			if (NumIntervals > 0)
 			{
-				const FChaosPointMovementPathPoint& PathPoint = PathPoints[PointIndex];
-				float DistanceOnPathSincePrevious = PathPoint.PathDistanceFromStart - PreviousPathDistance;
-				if (DistanceOnPathSincePrevious < SMALL_NUMBER)
+				float PreviousProgress = 0.0f;
+				float PreviousPathDistance = 0.0f;
+				float TotalProgressOnCollapsedIntervals = FMath::Clamp(float(NumCollapsedIntervals) / float(NumIntervals), 0.0f, 1.0f);
+				float TotalProgressOnSeparatedIntervals = FMath::Clamp(1.0f - TotalProgressOnCollapsedIntervals, 0.0f, 1.0f);
+				float ProgressOnOneCollapsedInterval = (TotalProgressOnCollapsedIntervals > SMALL_NUMBER) ? (1.0f / float(NumIntervals)) : 0.0f;
+				for (int PointIndex = 1; PointIndex < PathPoints.Num(); ++PointIndex)
 				{
-					// This is the end point of a collapsed interval, we increase progress by ProgressOnOneCollapsedInterval
-					PathPoint.Progress = PreviousProgress + ProgressOnOneCollapsedInterval;
-				}
-				else
-				{
-					// This is the end point of a non collapsed interval, we increase progress by TotalProgressOnSeparatedIntervals * d(i)/D
-					PathPoint.Progress = PreviousProgress + UKismetMathLibrary::SafeDivide(DistanceOnPathSincePrevious, TotalPathDistance) * TotalProgressOnSeparatedIntervals;
-				}
+					const FChaosPointMovementPathPoint& PathPoint = PathPoints[PointIndex];
+					float DistanceOnPathSincePrevious = PathPoint.PathDistanceFromStart - PreviousPathDistance;
+					if (DistanceOnPathSincePrevious < SMALL_NUMBER)
+					{
+						// This is the end point of a collapsed interval, we increase progress by ProgressOnOneCollapsedInterval
+						PathPoint.Progress = PreviousProgress + ProgressOnOneCollapsedInterval;
+					}
+					else
+					{
+						// This is the end point of a non collapsed interval, we increase progress by TotalProgressOnSeparatedIntervals * d(i)/D
+						PathPoint.Progress = PreviousProgress + UKismetMathLibrary::SafeDivide(DistanceOnPathSincePrevious, TotalPathDistance) * TotalProgressOnSeparatedIntervals;
+					}
 
-				PathPoint.Progress = FMath::Clamp(PathPoint.Progress, 0.0f, 1.0f);
-				PreviousProgress = PathPoint.Progress;
-				PreviousPathDistance = PathPoint.PathDistanceFromStart;
+					PathPoint.Progress = FMath::Clamp(PathPoint.Progress, 0.0f, 1.0f);
+					PreviousProgress = PathPoint.Progress;
+					PreviousPathDistance = PathPoint.PathDistanceFromStart;
+				}
+				// The last point needs to be at exactly 1.0 progress or some assumptions don't hold
+				PathPoints[PathPoints.Num()-1].Progress = 1.0f;
 			}
+
+			bHasAssignedPointProgress = true;
 		}
-		
-		bHasAssignedPointProgress = true;
 	}
+
+	CachedPathBasisTransform = PathBasisTransform;
 }
diff --git a/Engine/Plugins/Experimental/ChaosMover/Source/ChaosMover/Private/PathedMovement/ChaosSplineMovementPathPattern.cpp b/Engine/Plugins/Experimental/ChaosMover/Source/ChaosMover/Private/PathedMovement/ChaosSplineMovementPathPattern.cpp
--- a/Engine/Plugins/Experimental/ChaosMover/Source/ChaosMover/Private/PathedMovement/ChaosSplineMovementPathPattern.cpp
+++ b/Engine/Plugins/Experimental/ChaosMover/Source/ChaosMover/Private/PathedMovement/ChaosSplineMovementPathPattern.cpp
@@ -32,7 +32,7 @@ void UChaosSplineMovementPathPattern::InitializePattern(UChaosMoverSimulation* I
 	}
 }
 
-FTransform UChaosSplineMovementPathPattern::CalcUnmaskedTargetRelativeTransform(float PatternProgress, const FTransform& CurTargetTransform) const
+FTransform UChaosSplineMovementPathPattern::CalcUnmaskedTargetTransform(float PatternProgress, const FTransform& BasisTransform) const
 {
 	if (SplineComp)
 	{
@@ -40,7 +40,7 @@ FTransform UChaosSplineMovementPathPattern::CalcUnmaskedTargetRelativeTransform(
 
 		const bool bHasValidBounds = HasValidBounds();
 		const float SplineStartDist = bHasValidBounds ? SplineLength * LowerBound : 0.f;
-		const float SplineEndDist = bHasValidBounds ? SplineLength * UpperBound : 1.f;
+		const float SplineEndDist = bHasValidBounds ? SplineLength * UpperBound : SplineLength;
 
 		const float DistanceAlongSpline = FMath::Lerp(SplineStartDist, SplineEndDist, PatternProgress);
 		FTransform TargetToSpline = SplineComp->GetTransformAtDistanceAlongSpline(DistanceAlongSpline, ESplineCoordinateSpace::Local, bApplySplineScaling);
@@ -51,7 +51,10 @@ FTransform UChaosSplineMovementPathPattern::CalcUnmaskedTargetRelativeTransform(
 		}
 		
 		// Intentionally not accounting for any offset between the path origin and the spline's transform - zero progress should always equate to the path origin
-		return TargetToSpline;
+		return FTransform(
+			BasisTransform.TransformRotation(TargetToSpline.GetRotation()),
+			BasisTransform.TransformPositionNoScale(TargetToSpline.GetLocation()),
+			TargetToSpline.GetScale3D()* BasisTransform.GetScale3D());
 	}
-	return FTransform::Identity;
+	return BasisTransform;
 }
diff --git a/Engine/Plugins/Experimental/ChaosMover/Source/ChaosMover/Public/ChaosMover/PathedMovement/ChaosEllipticalMovementPathPattern.h b/Engine/Plugins/Experimental/ChaosMover/Source/ChaosMover/Public/ChaosMover/PathedMovement/ChaosEllipticalMovementPathPattern.h
--- a/Engine/Plugins/Experimental/ChaosMover/Source/ChaosMover/Public/ChaosMover/PathedMovement/ChaosEllipticalMovementPathPattern.h
+++ b/Engine/Plugins/Experimental/ChaosMover/Source/ChaosMover/Public/ChaosMover/PathedMovement/ChaosEllipticalMovementPathPattern.h
@@ -16,7 +16,7 @@ public:
 	virtual void InitializePattern(UChaosMoverSimulation* InSimulation) override;
 
 protected:
-	virtual FTransform CalcUnmaskedTargetRelativeTransform(float OverallPathFraction, const FTransform& CurTargetTransform) const override;
+	virtual FTransform CalcUnmaskedTargetTransform(float OverallPathFraction, const FTransform& BasisTransform) const override;
 
 	/** Radius of the ellipse along the x axis */
 	UPROPERTY(EditAnywhere, Category = EllipticalPathPattern)
diff --git a/Engine/Plugins/Experimental/ChaosMover/Source/ChaosMover/Public/ChaosMover/PathedMovement/ChaosPathedMovementMode.h b/Engine/Plugins/Experimental/ChaosMover/Source/ChaosMover/Public/ChaosMover/PathedMovement/ChaosPathedMovementMode.h
--- a/Engine/Plugins/Experimental/ChaosMover/Source/ChaosMover/Public/ChaosMover/PathedMovement/ChaosPathedMovementMode.h
+++ b/Engine/Plugins/Experimental/ChaosMover/Source/ChaosMover/Public/ChaosMover/PathedMovement/ChaosPathedMovementMode.h
@@ -40,8 +40,8 @@ public:
 	virtual bool ShouldUseConstraint() const override;
 	virtual const Chaos::FPBDJointSettings& GetConstraintSettings() const override;
 
-	// IChaosProgressingPathedMovementModeInterface
-	virtual Chaos::FRigidTransform3 CalcTargetRelativeTransform(float OverallPathProgress) const override;
+	// IChaosPathedMovementModeInterface
+	virtual FTransform CalcTargetTransform(float OverallPathProgress, const FTransform& BasisTransform) const override;
 
 
 #if WITH_EDITOR
diff --git a/Engine/Plugins/Experimental/ChaosMover/Source/ChaosMover/Public/ChaosMover/PathedMovement/ChaosPathedMovementPatternBase.h b/Engine/Plugins/Experimental/ChaosMover/Source/ChaosMover/Public/ChaosMover/PathedMovement/ChaosPathedMovementPatternBase.h
--- a/Engine/Plugins/Experimental/ChaosMover/Source/ChaosMover/Public/ChaosMover/PathedMovement/ChaosPathedMovementPatternBase.h
+++ b/Engine/Plugins/Experimental/ChaosMover/Source/ChaosMover/Public/ChaosMover/PathedMovement/ChaosPathedMovementPatternBase.h
@@ -44,7 +44,8 @@ public:
 
 	UChaosPathedMovementMode& GetMovementMode() const;
 	
-	FTransform CalcTargetRelativeTransform(float OverallPathProgress, const FTransform& CurTargetTransform) const;
+	// Calculates the path target transform in world coordinates, at OverallPathProgress, given the path basis transform
+	FTransform CalcTargetTransform(float OverallPathProgress, const FTransform& BasisTransform) const;
 
 	virtual bool DebugDrawUsingStepSamples() const { return bDebugDrawPattern; }
 	virtual void AppendDebugDrawElements(UChaosPathedMovementDebugDrawComponent& DebugDrawComp, FBoxSphereBounds::Builder& InOutDebugBoundsBuilder) {}
@@ -58,15 +59,16 @@ public:
 	
 protected:
 	/**
-	* Where patterns calculate their target transforms at the given path progress, optionally relative to the current aggregate target.
+	* Where patterns calculate their target transforms at the given path progress
+	* The transform returned is in world coordinates, given the path basis transform BasisTransform
 	* Axis masking is applied to this result before it's actually used.
 	*/
-	virtual FTransform CalcUnmaskedTargetRelativeTransform(float PatternProgress, const FTransform& CurTargetTransform) const { return FTransform::Identity; }
+	virtual FTransform CalcUnmaskedTargetTransform(float PatternProgress, const FTransform& BasisTransform) const { return BasisTransform; }
 
-	/** Utility function that calls CalcUnmaskedTargetRelativeTransform and then applies masking. This is mostly for debug drawing, which wants to
+	/** Utility function that calls CalcUnmaskedTargetTransform and then applies masking. This is mostly for debug drawing, which wants to
 	*   deal with PatternProgress and not OverallPathProgress but also get the masked transform
 	*/
-	FTransform CalcMaskedTargetRelativeTransform(float PatternProgress, const FTransform& CurTargetTransform) const;
+	FTransform CalcMaskedTargetTransform(float PatternProgress, const FTransform& BasisTransform) const;
 
 	/** Converts an overall path progress to a progress on the pattern
 	*   Returns -1 if it can't do it, say, e.g. if somehow StartAtPathProgress > EndAtPathProgress
diff --git a/Engine/Plugins/Experimental/ChaosMover/Source/ChaosMover/Public/ChaosMover/PathedMovement/ChaosPathedMovementTypes.h b/Engine/Plugins/Experimental/ChaosMover/Source/ChaosMover/Public/ChaosMover/PathedMovement/ChaosPathedMovementTypes.h
--- a/Engine/Plugins/Experimental/ChaosMover/Source/ChaosMover/Public/ChaosMover/PathedMovement/ChaosPathedMovementTypes.h
+++ b/Engine/Plugins/Experimental/ChaosMover/Source/ChaosMover/Public/ChaosMover/PathedMovement/ChaosPathedMovementTypes.h
@@ -165,24 +165,12 @@ struct FChaosPathedMovementModeDebugData : public FMoverDataStructBase
 	UPROPERTY(VisibleAnywhere, Category = "Mover Info")
 	FChaosPathedMovementState PostSimMoveState;
 
-	UPROPERTY(VisibleAnywhere, Category = "Mover Info")
-	FVector PreviousRelativeTranslation = FVector::ZeroVector;
-
-	UPROPERTY(VisibleAnywhere, Category = "Mover Info")
-	FRotator PreviousRelativeRotation = FRotator::ZeroRotator;
-
 	UPROPERTY(VisibleAnywhere, Category = "Mover Info")
 	FVector PreviousWorldTranslation = FVector::ZeroVector;
 
 	UPROPERTY(VisibleAnywhere, Category = "Mover Info")
 	FRotator PreviousWorldRotation = FRotator::ZeroRotator;
 
-	UPROPERTY(VisibleAnywhere, Category = "Mover Info")
-	FVector NewRelativeTranslation = FVector::ZeroVector;
-
-	UPROPERTY(VisibleAnywhere, Category = "Mover Info")
-	FRotator NewRelativeRotation = FRotator::ZeroRotator;
-
 	UPROPERTY(VisibleAnywhere, Category = "Mover Info")
 	FVector NewWorldTranslation = FVector::ZeroVector;
 
@@ -361,5 +349,5 @@ class IChaosPathedMovementModeInterface
 public:
 
 	// Returns the transform on the path at Progression
-	virtual Chaos::FRigidTransform3 CalcTargetRelativeTransform(float Progression) const = 0;
+	virtual FTransform CalcTargetTransform(float Progression, const FTransform& BasisTransform) const = 0;
 };
diff --git a/Engine/Plugins/Experimental/ChaosMover/Source/ChaosMover/Public/ChaosMover/PathedMovement/ChaosPointMovementPathPattern.h b/Engine/Plugins/Experimental/ChaosMover/Source/ChaosMover/Public/ChaosMover/PathedMovement/ChaosPointMovementPathPattern.h
--- a/Engine/Plugins/Experimental/ChaosMover/Source/ChaosMover/Public/ChaosMover/PathedMovement/ChaosPointMovementPathPattern.h
+++ b/Engine/Plugins/Experimental/ChaosMover/Source/ChaosMover/Public/ChaosMover/PathedMovement/ChaosPointMovementPathPattern.h
@@ -46,6 +46,8 @@ public:
 	UPROPERTY(EditAnywhere, BlueprintReadOnly, Category = Location, meta = (MakeEditWidget = true))
 	FRotator Rotation = FRotator::ZeroRotator;
 
+	FTransform GetWorldTransform(const FTransform& BasisTransform) const;
+
 	/**
 	 * The progress along the path that this point corresponds to.
 	 * Read-only, calculated automatically based on number of points, distances between them, and movement rules for the point 
@@ -57,8 +59,18 @@ public:
 	// UPROPERTY(VisibleAnywhere, BlueprintReadOnly, Transient, Category = Location)
 	mutable float PathDistanceFromStart = 0.f;
 
-	mutable FVector WorldLoc = FVector::ZeroVector;
-	mutable FQuat WorldRot = FQuat::Identity;
+	// This location is precomputed as a function of Location and Basis.
+	// If Basis is World, then this is a world location. Otherwise this location is relative to the BasisTransform that will be provided (@see CalcUnmaskedTargetTransform)
+	mutable FVector BasedLocation = FVector::ZeroVector;
+	// This rotation is precomputed as a function of Rotation and Basis.
+	// If RotationBasis is World, then this is a world rotation. Otherwise this rotation is relative to the BasisTransform that will be provided (@see CalcUnmaskedTargetTransform)
+	mutable FQuat BasedRotation = FQuat::Identity;
+	// This basis indicates whether the BasedLocation is relative to the path basis or in world coordinates
+	// It is either PathOrigin or World, but cannot be PreviousPoint, since the purpose of BasedLocation is to precompute as much as possible beforehand
+	mutable EChaosPointMovementLocationBasis EffectiveBasis = EChaosPointMovementLocationBasis::PathOrigin;
+	// This basis indicates whether the BasedRotation is relative to the path basis or in world coordinates
+	// It is either PathOrigin or World, but cannot be PreviousPoint, since the purpose of BasedRotation is to precompute as much as possible beforehand
+	mutable EChaosPointMovementLocationBasis EffectiveRotationBasis = EChaosPointMovementLocationBasis::PathOrigin;
 };
 
 /** Movement pattern that moves between explicitly defined points */
@@ -74,14 +86,13 @@ public:
 	virtual void PostEditChangeProperty(struct FPropertyChangedEvent& PropertyChangedEvent) override;
 #endif
 	
-	virtual bool DebugDrawUsingStepSamples() const override { return false; }
 	virtual void AppendDebugDrawElements(UChaosPathedMovementDebugDrawComponent& DebugDrawComp, FBoxSphereBounds::Builder& InOutDebugBoundsBuilder) override;
 
 protected:
-	virtual FTransform CalcUnmaskedTargetRelativeTransform(float PatternProgress, const FTransform& CurTargetTransform) const override;
+	virtual FTransform CalcUnmaskedTargetTransform(float PatternProgress, const FTransform& BasisTransform) const override;
 
-	//@todo DanH: We need to force a refresh whenever the path origin changes
-	void RefreshAssignedPointProgress(bool bForceRefresh = false) const;
+	void RefreshAssignedPointProgress(const FTransform& PathBasisTransform) const;
+	FTransform GetPathBasisTransform() const;
 
 	/** Relative point locations to move toward, in sequence from first to last */
 	UPROPERTY(EditAnywhere, BlueprintReadWrite, Category = "Point Movement Pattern")
@@ -92,4 +103,6 @@ protected:
 	mutable float TotalPathDistance = 0.f;
 
 	mutable bool bHasAssignedPointProgress = false;
+
+	mutable FTransform CachedPathBasisTransform = FTransform::Identity;
 };
\ No newline at end of file
diff --git a/Engine/Plugins/Experimental/ChaosMover/Source/ChaosMover/Public/ChaosMover/PathedMovement/ChaosSplineMovementPathPattern.h b/Engine/Plugins/Experimental/ChaosMover/Source/ChaosMover/Public/ChaosMover/PathedMovement/ChaosSplineMovementPathPattern.h
--- a/Engine/Plugins/Experimental/ChaosMover/Source/ChaosMover/Public/ChaosMover/PathedMovement/ChaosSplineMovementPathPattern.h
+++ b/Engine/Plugins/Experimental/ChaosMover/Source/ChaosMover/Public/ChaosMover/PathedMovement/ChaosSplineMovementPathPattern.h
@@ -35,7 +35,7 @@ public:
 	bool HasValidBounds() const { return UpperBound > LowerBound; }
 	
 protected:
-	virtual FTransform CalcUnmaskedTargetRelativeTransform(float PatternProgress, const FTransform& CurTargetTransform) const override;
+	virtual FTransform CalcUnmaskedTargetTransform(float PatternProgress, const FTransform& BasisTransform) const override;
 
 	/** How far into the spline the movement path actually begins.  */
 	UPROPERTY(EditAnywhere, Category = SplinePathPattern, meta = (UIMin = 0.f, ClampMin = 0.f, UIMax = 1.f, ClampMax = 1.f))
