ChaosMover - Adjust floor distance when on slope to ensure distance is measured underneath the middle of the character

diff --git a/Engine/Plugins/Experimental/ChaosMover/Source/ChaosMover/Private/Utilities/ChaosMoverQueryUtils.cpp b/Engine/Plugins/Experimental/ChaosMover/Source/ChaosMover/Private/Utilities/ChaosMoverQueryUtils.cpp
--- a/Engine/Plugins/Experimental/ChaosMover/Source/ChaosMover/Private/Utilities/ChaosMoverQueryUtils.cpp
+++ b/Engine/Plugins/Experimental/ChaosMover/Source/ChaosMover/Private/Utilities/ChaosMoverQueryUtils.cpp
@@ -135,6 +135,15 @@ namespace UE::ChaosMover::Utils
 			OutFloorResult.FloorDist = Params.UpDir.Dot(Params.Location - OutHit.ImpactPoint);
 			OutFloorResult.HitResult = OutHit;
 
+			// Update the ground distance based on the slope. If you are on a slope the query might hit on an edge
+			// rather than directly under the capsule. Also move back to the start location
+			const float DP = OutHit.ImpactNormal.Dot(Params.UpDir);
+			if (DP > UE_SMALL_NUMBER)
+			{
+				OutFloorResult.FloorDist += 2.0f * Params.QueryRadius * (1.0f - DP) / DP
+					+ FVector::VectorPlaneProject(Params.DeltaPos, OutHit.ImpactNormal).Dot(Params.UpDir);
+			}
+
 			// Make sure that the object stored in the hit result is the root object
 			if (const Chaos::FPhysicsObjectHandle HitObject = OutHit.PhysicsObject)
 			{
