ChaosMover - Add option to zero friction on contacts when moving

diff --git a/Engine/Plugins/Experimental/ChaosMover/Source/ChaosMover/Private/Character/Modes/ChaosCharacterMovementMode.cpp b/Engine/Plugins/Experimental/ChaosMover/Source/ChaosMover/Private/Character/Modes/ChaosCharacterMovementMode.cpp
--- a/Engine/Plugins/Experimental/ChaosMover/Source/ChaosMover/Private/Character/Modes/ChaosCharacterMovementMode.cpp
+++ b/Engine/Plugins/Experimental/ChaosMover/Source/ChaosMover/Private/Character/Modes/ChaosCharacterMovementMode.cpp
@@ -3,8 +3,11 @@
 #include "ChaosMover/Character/Modes/ChaosCharacterMovementMode.h"
 
 #include "Chaos/Character/CharacterGroundConstraintSettings.h"
+#include "Chaos/ContactModification.h"
+#include "Chaos/PhysicsObjectInternalInterface.h"
 #include "ChaosMover/Character/Settings/SharedChaosCharacterMovementSettings.h"
 #include "ChaosMover/ChaosMoverLog.h"
+#include "ChaosMover/ChaosMoverSimulation.h"
 #include "Components/CapsuleComponent.h"
 #include "Components/SkeletalMeshComponent.h"
 #include "GameFramework/Character.h"
@@ -151,3 +154,60 @@ void UChaosCharacterMovementMode::ClearAccelerationOverride()
 {
 	AccelerationOverride.Reset();
 }
+
+void UChaosCharacterMovementMode::ModifyContacts(const FMoverTimeStep& TimeStep, const UE::ChaosMover::FSimulationInputData& InputData, const UE::ChaosMover::FSimulationOutputData& OutputData, Chaos::FCollisionContactModifier& Modifier) const
+{
+	Super::ModifyContacts(TimeStep, InputData, OutputData, Modifier);
+
+	if (FrictionOverrideMode == ECharacterMoverFrictionOverrideMode::DoNotOverride)
+	{
+		return;
+	}
+
+	check(Simulation);
+
+	// Get the updated (character) particle
+	Chaos::FGeometryParticleHandle* UpdatedParticle = nullptr;
+	const FChaosMoverSimulationDefaultInputs* SimInputs = Simulation->GetLocalSimInput().FindDataByType<FChaosMoverSimulationDefaultInputs>();
+	if (SimInputs)
+	{
+		Chaos::FReadPhysicsObjectInterface_Internal ReadInterface = Chaos::FPhysicsObjectInternalInterface::GetRead();
+		UpdatedParticle = ReadInterface.GetParticle(SimInputs->PhysicsObject);
+	}
+
+	if (!UpdatedParticle)
+	{
+		return;
+	}
+
+	for (Chaos::FContactPairModifier& PairModifier : Modifier.GetContacts(UpdatedParticle))
+	{
+		bool bOverrideToZero = false;
+
+		switch (FrictionOverrideMode)
+		{
+		case ECharacterMoverFrictionOverrideMode::AlwaysOverrideToZero:
+			bOverrideToZero = true;
+			break;
+		case ECharacterMoverFrictionOverrideMode::OverrideToZeroWhenMoving:
+		{
+			const FCharacterDefaultInputs* CharacterInputs = InputData.InputCmd.InputCollection.FindDataByType<FCharacterDefaultInputs>();
+			if (CharacterInputs)
+			{
+				constexpr float MinInput = 0.1f;
+				bOverrideToZero = CharacterInputs->GetMoveInput().SizeSquared() > MinInput* MinInput;
+			}
+			break;
+		}
+			
+		default:
+			break;
+		}
+
+		if (bOverrideToZero)
+		{
+			PairModifier.ModifyStaticFriction(0.0f);
+			PairModifier.ModifyDynamicFriction(0.0f);
+		}
+	}
+}
diff --git a/Engine/Plugins/Experimental/ChaosMover/Source/ChaosMover/Private/Character/Modes/ChaosWalkingMode.cpp b/Engine/Plugins/Experimental/ChaosMover/Source/ChaosMover/Private/Character/Modes/ChaosWalkingMode.cpp
--- a/Engine/Plugins/Experimental/ChaosMover/Source/ChaosMover/Private/Character/Modes/ChaosWalkingMode.cpp
+++ b/Engine/Plugins/Experimental/ChaosMover/Source/ChaosMover/Private/Character/Modes/ChaosWalkingMode.cpp
@@ -403,6 +403,8 @@ void UChaosWalkingMode::GetFloorAndCheckMovement(const FMoverDefaultSyncState& S
 
 void UChaosWalkingMode::ModifyContacts(const FMoverTimeStep& TimeStep, const UE::ChaosMover::FSimulationInputData& InputData, const UE::ChaosMover::FSimulationOutputData& OutputData, Chaos::FCollisionContactModifier& Modifier) const
 {
+	Super::ModifyContacts(TimeStep, InputData, OutputData, Modifier);
+
 	check(Simulation);
 
 	// Get the updated (character) particle
diff --git a/Engine/Plugins/Experimental/ChaosMover/Source/ChaosMover/Public/ChaosMover/ChaosMovementMode.h b/Engine/Plugins/Experimental/ChaosMover/Source/ChaosMover/Public/ChaosMover/ChaosMovementMode.h
--- a/Engine/Plugins/Experimental/ChaosMover/Source/ChaosMover/Public/ChaosMover/ChaosMovementMode.h
+++ b/Engine/Plugins/Experimental/ChaosMover/Source/ChaosMover/Public/ChaosMover/ChaosMovementMode.h
@@ -58,7 +58,7 @@ public:
 	{
 	}
 
-	UPROPERTY(EditAnywhere, Category = Collision)
+	UPROPERTY(EditAnywhere, Category = "Collision Settings")
 	EChaosMoverIgnoredCollisionMode IgnoredCollisionMode = EChaosMoverIgnoredCollisionMode::DisableCollisionsWithIgnored;
 
 protected:
diff --git a/Engine/Plugins/Experimental/ChaosMover/Source/ChaosMover/Public/ChaosMover/Character/Modes/ChaosCharacterMovementMode.h b/Engine/Plugins/Experimental/ChaosMover/Source/ChaosMover/Public/ChaosMover/Character/Modes/ChaosCharacterMovementMode.h
--- a/Engine/Plugins/Experimental/ChaosMover/Source/ChaosMover/Public/ChaosMover/Character/Modes/ChaosCharacterMovementMode.h
+++ b/Engine/Plugins/Experimental/ChaosMover/Source/ChaosMover/Public/ChaosMover/Character/Modes/ChaosCharacterMovementMode.h
@@ -12,6 +12,14 @@ namespace Chaos
 	class FCharacterGroundConstraintSettings;
 }
 
+UENUM(BlueprintType)
+enum class ECharacterMoverFrictionOverrideMode : uint8
+{
+	DoNotOverride,
+	AlwaysOverrideToZero,
+	OverrideToZeroWhenMoving,
+};
+
 /**
  * Base class for all Chaos character movement modes
  */
@@ -34,6 +42,8 @@ public:
 		return true;
 	}
 
+	CHAOSMOVER_API virtual void ModifyContacts(const FMoverTimeStep& TimeStep, const UE::ChaosMover::FSimulationInputData& InputData, const UE::ChaosMover::FSimulationOutputData& OutputData, Chaos::FCollisionContactModifier& Modifier) const override;
+
 	CHAOSMOVER_API virtual void UpdateConstraintSettings(Chaos::FCharacterGroundConstraintSettings& ConstraintSettings) const override;
 
 	// IChaosCharacterMovementModeInterface
@@ -98,6 +108,10 @@ protected:
 	UPROPERTY(EditAnywhere, Category = "Query Settings", meta = (ClampMin = "0", UIMin = "0", ForceUnits = "cm"))
 	TOptional<float> QueryRadiusOverride;
 
+	// Allows the mode to override friction on collision with other physics bodies.
+	UPROPERTY(EditAnywhere, Category = "Collision Settings")
+	ECharacterMoverFrictionOverrideMode FrictionOverrideMode = ECharacterMoverFrictionOverrideMode::OverrideToZeroWhenMoving;
+
 private:
 	float TargetHeight = 95.0f;
 	float QueryRadius = 30.0f;
