ChaosMover: Fixes to bugs (some path movement transform edge cases) that BugHawk found after I submitted, even though I reviewed with BugHawk prior to submission and it didn't find it then.

diff --git a/Engine/Plugins/Experimental/ChaosMover/Source/ChaosMover/Private/PathedMovement/ChaosPathedMovementPatternBase.cpp b/Engine/Plugins/Experimental/ChaosMover/Source/ChaosMover/Private/PathedMovement/ChaosPathedMovementPatternBase.cpp
--- a/Engine/Plugins/Experimental/ChaosMover/Source/ChaosMover/Private/PathedMovement/ChaosPathedMovementPatternBase.cpp
+++ b/Engine/Plugins/Experimental/ChaosMover/Source/ChaosMover/Private/PathedMovement/ChaosPathedMovementPatternBase.cpp
@@ -134,7 +134,7 @@ FTransform UChaosPathedMovementPatternBase::CalcMaskedTargetTransform(float Patt
 {
 	if (PatternProgress < 0.0f)
 	{
-		return FTransform();
+		return BasisTransform;
 	}
 
 	const FTransform UnmaskedTarget = CalcUnmaskedTargetTransform(PatternProgress, BasisTransform);
diff --git a/Engine/Plugins/Experimental/ChaosMover/Source/ChaosMover/Private/PathedMovement/ChaosPointMovementPathPattern.cpp b/Engine/Plugins/Experimental/ChaosMover/Source/ChaosMover/Private/PathedMovement/ChaosPointMovementPathPattern.cpp
--- a/Engine/Plugins/Experimental/ChaosMover/Source/ChaosMover/Private/PathedMovement/ChaosPointMovementPathPattern.cpp
+++ b/Engine/Plugins/Experimental/ChaosMover/Source/ChaosMover/Private/PathedMovement/ChaosPointMovementPathPattern.cpp
@@ -75,8 +75,8 @@ void UChaosPointMovementPathPattern::AppendDebugDrawElements(UChaosPathedMovemen
 FTransform FChaosPointMovementPathPoint::GetWorldTransform(const FTransform& BasisTransform) const
 {
 	return FTransform(
-		(EffectiveRotationBasis == EChaosPointMovementLocationBasis::World) ? Rotation : FRotator(BasisTransform.TransformRotation(FQuat(BasedRotation))),
-		(EffectiveBasis == EChaosPointMovementLocationBasis::World) ? Location : BasisTransform.TransformPositionNoScale(BasedLocation),
+		(EffectiveRotationBasis == EChaosPointMovementLocationBasis::World) ? FRotator(BasedRotation) : FRotator(BasisTransform.TransformRotation(FQuat(BasedRotation))),
+		(EffectiveBasis == EChaosPointMovementLocationBasis::World) ? BasedLocation : BasisTransform.TransformPositionNoScale(BasedLocation),
 		BasisTransform.GetScale3D()
 	);
 }
