Chaos Mover - Add option to apply friction to lateral movement when falling

diff --git a/Engine/Plugins/Experimental/ChaosMover/Source/ChaosMover/Private/Character/Modes/ChaosFallingMode.cpp b/Engine/Plugins/Experimental/ChaosMover/Source/ChaosMover/Private/Character/Modes/ChaosFallingMode.cpp
--- a/Engine/Plugins/Experimental/ChaosMover/Source/ChaosMover/Private/Character/Modes/ChaosFallingMode.cpp
+++ b/Engine/Plugins/Experimental/ChaosMover/Source/ChaosMover/Private/Character/Modes/ChaosFallingMode.cpp
@@ -19,6 +19,7 @@ UChaosFallingMode::UChaosFallingMode(const FObjectInitializer& ObjectInitializer
 	, bCancelVerticalSpeedOnLanding(true)
 	, AirControlPercentage(0.4f)
 	, FallingDeceleration(200.0f)
+	, FallingLateralFriction(0.0f)
 	, OverTerminalSpeedFallingDeceleration(800.0f)
 	, TerminalMovementPlaneSpeed(1500.0f)
 	, bShouldClampTerminalVerticalSpeed(true)
@@ -97,6 +98,7 @@ void UChaosFallingMode::GenerateMove_Implementation(const FMoverTickStartData& S
 	Params.Deceleration = FallingDeceleration;
 	Params.WorldToGravityQuat = WorldToGravityTransform;
 	Params.bUseAccelerationForVelocityMove = SharedSettings->bUseAccelerationForVelocityMove;
+	Params.Friction = FallingLateralFriction;
 
 	// Check if any current velocity values are over our terminal velocity - if so limit the move input in that direction and apply OverTerminalVelocityFallingDeceleration
 	if (Params.MoveInput.Dot(StartVelocity) > 0 && StartHorizontalVelocity.Size() >= TerminalMovementPlaneSpeed)
diff --git a/Engine/Plugins/Experimental/ChaosMover/Source/ChaosMover/Public/ChaosMover/Character/Modes/ChaosFallingMode.h b/Engine/Plugins/Experimental/ChaosMover/Source/ChaosMover/Public/ChaosMover/Character/Modes/ChaosFallingMode.h
--- a/Engine/Plugins/Experimental/ChaosMover/Source/ChaosMover/Public/ChaosMover/Character/Modes/ChaosFallingMode.h
+++ b/Engine/Plugins/Experimental/ChaosMover/Source/ChaosMover/Public/ChaosMover/Character/Modes/ChaosFallingMode.h
@@ -43,6 +43,15 @@ public:
 	UPROPERTY(Category = Mover, EditAnywhere, BlueprintReadWrite, meta = (ClampMin = "0", ForceUnits = "cm/s^2"))
 	float FallingDeceleration;
 
+	/**
+	 * Friction to apply to air movement.
+	 * Lateral velocity is scaled each tick by a factor (1-Friction * TimeStep) so friction
+	 * values greater than 1/TimeStep will result in all lateral velocity being removed.
+	 * Note: This is NOT applied to vertical velocity, only movement plane velocity
+	 */
+	UPROPERTY(Category = Mover, EditAnywhere, BlueprintReadWrite, meta = (ClampMin = "0"))
+	float FallingLateralFriction;
+
 	/**
 	 * Deceleration to apply to air movement when falling faster than terminal velocity
 	 * Note: This is NOT applied to vertical velocity, only movement plane velocity
