PBCM: Fix to prevent movement mode correction event being lost after a resim.
#rnx

diff --git a/Engine/Plugins/Experimental/ChaosMover/Source/ChaosMover/Private/ChaosMoverStateMachine.cpp b/Engine/Plugins/Experimental/ChaosMover/Source/ChaosMover/Private/ChaosMoverStateMachine.cpp
--- a/Engine/Plugins/Experimental/ChaosMover/Source/ChaosMover/Private/ChaosMoverStateMachine.cpp
+++ b/Engine/Plugins/Experimental/ChaosMover/Source/ChaosMover/Private/ChaosMoverStateMachine.cpp
@@ -629,12 +629,13 @@ namespace UE::ChaosMover
 		float PreviousInternalSimTimeMs = InternalSimTimeMs;
 		int32 PreviousInternalServerFrame = InternalServerFrame;
 		InternalSimTimeMs = NewTimeStep.BaseSimTimeMs;
+		InternalServerFrame = NewTimeStep.ServerFrame;
 
 		ClearQueuedMode();
 
 		if (CurrentModeName != NewSyncState.MovementMode)
 		{
-			SetModeImmediately(NewSyncState.MovementMode);
+			QueueNextMode(NewSyncState.MovementMode);
 		}
 
 		QueuedLayeredMoves.Empty();
