ChaosMover - Fixes to constraint initialization to prevent crash on destroy physics

#rnx

diff --git a/Engine/Plugins/Experimental/ChaosMover/Source/ChaosMover/Private/Backends/ChaosMoverBackend.cpp b/Engine/Plugins/Experimental/ChaosMover/Source/ChaosMover/Private/Backends/ChaosMoverBackend.cpp
--- a/Engine/Plugins/Experimental/ChaosMover/Source/ChaosMover/Private/Backends/ChaosMoverBackend.cpp
+++ b/Engine/Plugins/Experimental/ChaosMover/Source/ChaosMover/Private/Backends/ChaosMoverBackend.cpp
@@ -238,10 +238,12 @@ void UChaosMoverBackendComponent::HandleUpdatedComponentPhysicsStateChanged(UPri
 	if (StateChange == EComponentPhysicsStateChange::Destroyed)
 	{
 		bWantsDestroySim = true;
+		DestroyPhysics();
 	}
 	else if (StateChange == EComponentPhysicsStateChange::Created)
 	{
 		bWantsCreateSim = true;
+		CreatePhysics();
 	}
 }
 
@@ -282,6 +284,37 @@ void UChaosMoverBackendComponent::InitSimulation()
 {
 	UMoverComponent& MoverComp = GetMoverComponent();
 
+	Chaos::FCharacterGroundConstraintHandle* CharacterConstraintHandle = nullptr;
+	if (CharacterGroundConstraint)
+	{
+		if (Chaos::FCharacterGroundConstraintProxy* Proxy = CharacterGroundConstraint->GetProxy<Chaos::FCharacterGroundConstraintProxy>())
+		{
+			CharacterConstraintHandle = Proxy->GetPhysicsThreadAPI();
+		}
+	}
+	if (!CharacterConstraintHandle)
+	{
+		return;
+	}
+
+	Chaos::FPBDJointConstraintHandle* JointConstraintHandle = nullptr;
+	Chaos::FKinematicGeometryParticleHandle* JointEndPointParticle = nullptr;
+	if (ActuationConstraintHandle.IsValid())
+	{
+		if (Chaos::FJointConstraintPhysicsProxy* Proxy = ActuationConstraintHandle->GetProxy<Chaos::FJointConstraintPhysicsProxy>())
+		{
+			JointConstraintHandle = Proxy->GetHandle();
+			if (Chaos::FSingleParticlePhysicsProxy* EndPointProxy = Proxy->GetConstraint()->GetKinematicEndPoint())
+			{
+				JointEndPointParticle = EndPointProxy->GetHandle_LowLevel()->CastToKinematicParticle();
+			}
+		}
+	}
+	if (!JointConstraintHandle || !JointEndPointParticle)
+	{
+		return;
+	}
+
 	UChaosMoverSimulation::FInitParams Params;
 	for (const TPair<FName, TObjectPtr<UBaseMovementMode>>& Pair : MoverComp.MovementModes)
 	{
@@ -295,9 +328,9 @@ void UChaosMoverBackendComponent::InitSimulation()
 	Params.ImmediateModeTransition = TWeakObjectPtr<UImmediateMovementModeTransition>(ImmediateModeTransition.Get());
 	Params.NullMovementMode = TWeakObjectPtr<UNullMovementMode>(NullMovementMode.Get());
 	Params.StartingMovementMode = MoverComp.StartingMovementMode;
-	Params.CharacterConstraintProxy = CharacterGroundConstraint ? CharacterGroundConstraint->GetProxy<Chaos::FCharacterGroundConstraintProxy>() : nullptr;
-	Params.ActuationConstraintProxy = ActuationConstraintHandle.IsValid() ? ActuationConstraintHandle->GetProxy<Chaos::FJointConstraintPhysicsProxy>() : nullptr;
-	Params.ActuationConstraintKinematicEndPointProxy = Params.ActuationConstraintProxy && Params.ActuationConstraintProxy->GetConstraint() ? Params.ActuationConstraintProxy->GetConstraint()->GetKinematicEndPoint() : nullptr;
+	Params.CharacterConstraintHandle = CharacterConstraintHandle;
+	Params.ActuationConstraintHandle = JointConstraintHandle;
+	Params.ActuationConstraintEndPointParticleHandle = JointEndPointParticle;
 	Params.TransformOnInit = TransformOnInit;
 	Params.PhysicsObject = GetPhysicsObject();
 	Params.Solver = GetPhysicsSolver();
@@ -322,11 +355,14 @@ void UChaosMoverBackendComponent::InitSimulation()
 	Params.InitialSyncState = OutputData.SyncState;
 
 	Simulation->Init(Params);
+
+	bWantsCreateSim = false;
 }
 
 void UChaosMoverBackendComponent::DeinitSimulation()
 {
 	Simulation->Deinit();
+	bWantsDestroySim = false;
 }
 
 void UChaosMoverBackendComponent::BeginPlay()
@@ -336,7 +372,6 @@ void UChaosMoverBackendComponent::BeginPlay()
 	if (UWorld* World = GetWorld(); World && World->IsGameWorld())
 	{
 		CreatePhysics();
-		InitSimulation();
 
 		// Register with the world subsystem
 		if (UChaosMoverSubsystem* ChaosMoverSubsystem = UWorld::GetSubsystem<UChaosMoverSubsystem>(GetWorld()))
@@ -430,17 +465,17 @@ void UChaosMoverBackendComponent::ProduceInputData(int32 PhysicsStep, int32 NumS
 	if (bWantsDestroySim)
 	{
 		DeinitSimulation();
-		DestroyPhysics();
-		bWantsDestroySim = false;
 		return;
 	}
 	if (bWantsCreateSim)
 	{
-		CreatePhysics();
 		InitSimulation();
-		bWantsCreateSim = false;
 	}
 
+	if (bWantsCreateSim)
+	{
+		return;
+	}
 
 	bool bIsGeneratingInputsLocally = !NetworkPhysicsComponent || NetworkPhysicsComponent->IsLocallyControlled();
 	if (bIsGeneratingInputsLocally)
@@ -557,6 +592,11 @@ void UChaosMoverBackendComponent::ConsumeOutputData(const FMoverTimeStep& TimeSt
 {
 	Chaos::EnsureIsInGameThreadContext();
 
+	if (bWantsCreateSim)
+	{
+		return;
+	}
+
 	SimOutputRecord.Add(TimeStep, OutputData);
 }
 
@@ -564,6 +604,11 @@ void UChaosMoverBackendComponent::FinalizeFrame(double ResultsTimeInMs)
 {
 	Chaos::EnsureIsInGameThreadContext();
 
+	if (bWantsCreateSim)
+	{
+		return;
+	}
+
 	UMoverComponent& MoverComp = GetMoverComponent();
 
 	FMoverTimeStep TimeStep;
diff --git a/Engine/Plugins/Experimental/ChaosMover/Source/ChaosMover/Private/Backends/ChaosMoverSubsystem.cpp b/Engine/Plugins/Experimental/ChaosMover/Source/ChaosMover/Private/Backends/ChaosMoverSubsystem.cpp
--- a/Engine/Plugins/Experimental/ChaosMover/Source/ChaosMover/Private/Backends/ChaosMoverSubsystem.cpp
+++ b/Engine/Plugins/Experimental/ChaosMover/Source/ChaosMover/Private/Backends/ChaosMoverSubsystem.cpp
@@ -116,16 +116,19 @@ void UChaosMoverSubsystem::InjectInputs_External(int32 PhysicsStep, int32 NumSte
 				AsyncInput->PhysicsSolver = Solver;
 
 				TimeStep.BaseSimTimeMs = Solver->GetPhysicsResultsTime_External() * 1000.0;
-				TimeStep.ServerFrame = Solver->GetCurrentFrame() + NetPhysicsTickOffset;
-
+				
 				if (Solver->IsUsingAsyncResults())
 				{
+					TimeStep.ServerFrame = FMath::FloorToInt(Solver->GetPhysicsResultsTime_External() / (Solver->GetAsyncDeltaTime()));
 					TimeStep.StepMs = Solver->GetAsyncDeltaTime() * 1000.0f;
 				}
 				else
 				{
+					TimeStep.ServerFrame = Solver->GetCurrentFrame();
 					TimeStep.StepMs = FMath::Clamp(World->GetDeltaSeconds(), Solver->GetMinDeltaTime_External(), Solver->GetMaxDeltaTime_External()) * 1000.0f;
 				}
+
+				TimeStep.ServerFrame += NetPhysicsTickOffset;
 			}
 		}
 	}
diff --git a/Engine/Plugins/Experimental/ChaosMover/Source/ChaosMover/Private/ChaosMoverSimulation.cpp b/Engine/Plugins/Experimental/ChaosMover/Source/ChaosMover/Private/ChaosMoverSimulation.cpp
--- a/Engine/Plugins/Experimental/ChaosMover/Source/ChaosMover/Private/ChaosMoverSimulation.cpp
+++ b/Engine/Plugins/Experimental/ChaosMover/Source/ChaosMover/Private/ChaosMoverSimulation.cpp
@@ -2,21 +2,23 @@
 
 #include "ChaosMover/ChaosMoverSimulation.h"
 
+#include "Chaos/ChaosEngineInterface.h"
 #include "Chaos/Character/CharacterGroundConstraint.h"
 #include "Chaos/Character/CharacterGroundConstraintContainer.h"
 #include "Chaos/ContactModification.h"
+#include "Chaos/KinematicTargets.h"
+#include "Chaos/PBDJointConstraintData.h"
 #include "Chaos/PhysicsObjectInternalInterface.h"
-#include "PhysicsEngine/PhysicsObjectExternalInterface.h"
-#include "Physics/Experimental/PhysScene_Chaos.h"
 #include "ChaosMover/ChaosMovementMode.h"
 #include "ChaosMover/ChaosMovementModeTransition.h"
+#include "ChaosMover/ChaosMoverConsoleVariables.h"
 #include "ChaosMover/ChaosMoverLog.h"
 #include "ChaosMover/ChaosMoverSimulationTypes.h"
 #include "ChaosMover/Character/ChaosCharacterInputs.h"
+#include "ChaosMover/PathedMovement/ChaosPathedMovementTypes.h"
 #include "ChaosMover/Utilities/ChaosGroundMovementUtils.h"
 #include "ChaosMover/Utilities/ChaosMoverQueryUtils.h"
-#include "ChaosMover/PathedMovement/ChaosPathedMovementTypes.h"
-#include "ChaosMover/ChaosMoverConsoleVariables.h"
+#include "Engine/World.h"
 #include "Framework/Threading.h"
 #include "MoveLibrary/FloorQueryUtils.h"
 #include "MoveLibrary/MovementMixer.h"
@@ -25,12 +27,10 @@
 #include "MovementModeStateMachine.h"
 #include "MoverSimulationTypes.h"
 #include "PBDRigidsSolver.h"
+#include "Physics/Experimental/PhysScene_Chaos.h"
+#include "PhysicsEngine/PhysicsObjectExternalInterface.h"
 #include "PhysicsProxy/CharacterGroundConstraintProxy.h"
-#include "Chaos/PBDJointConstraintData.h"
 #include "UObject/UObjectGlobals.h"
-#include "Chaos/ChaosEngineInterface.h"
-#include "Chaos/KinematicTargets.h"
-#include "Engine/World.h"
 
 #if WITH_CHAOS_VISUAL_DEBUGGER
 #include "ChaosVisualDebugger/ChaosVisualDebuggerTrace.h"
@@ -104,9 +104,9 @@ void UChaosMoverSimulation::Init(const FInitParams& InitParams)
 	Chaos::EnsureIsInGameThreadContext();
 
 	MovementMixerWeakPtr = InitParams.MovementMixer.Get();
-	CharacterConstraintProxy = InitParams.CharacterConstraintProxy;
-	ActuationConstraintProxy = InitParams.ActuationConstraintProxy;
-	ActuationConstraintKinematicEndPointProxy = InitParams.ActuationConstraintKinematicEndPointProxy;
+	CharacterConstraintHandle = InitParams.CharacterConstraintHandle;
+	ActuationConstraintHandle = InitParams.ActuationConstraintHandle;
+	ActuationConstraintEndPointParticleHandle = InitParams.ActuationConstraintEndPointParticleHandle;
 	PhysicsObject = InitParams.PhysicsObject;
 	Solver = InitParams.Solver;
 	World = InitParams.World;
@@ -157,10 +157,14 @@ void UChaosMoverSimulation::Init(const FInitParams& InitParams)
 	StateMachine.QueueNextMode(StateMachine.GetDefaultModeName());
 
 	OnInit();
+
+	bInitialized = true;
 }
 
 void UChaosMoverSimulation::Deinit()
 {
+	bInitialized = false;
+
 	OnDeinit();
 }
 
@@ -211,15 +215,23 @@ void UChaosMoverSimulation::SimulationTick(const FMoverTimeStep& TimeStep, const
 {
 	Chaos::EnsureIsInPhysicsThreadContext();
 
-	OnPreSimulationTick(TimeStep, InputData);
-	OnSimulationTick(TimeStep, InputData, OutputData);
-	OnPostSimulationTick(TimeStep, OutputData);
+	if (bInitialized)
+	{
+		OnPreSimulationTick(TimeStep, InputData);
+		OnSimulationTick(TimeStep, InputData, OutputData);
+		OnPostSimulationTick(TimeStep, OutputData);
+	}
 }
 
 void UChaosMoverSimulation::ModifyContacts(const FMoverTimeStep& TimeStep, const UE::ChaosMover::FSimulationInputData& InputData, const UE::ChaosMover::FSimulationOutputData& OutputData, Chaos::FCollisionContactModifier& Modifier)
 {
 	Chaos::EnsureIsInPhysicsThreadContext();
 
+	if (!bInitialized)
+	{
+		return;
+	}
+
 	if (TStrongObjectPtr<const UBaseMovementMode> CurrentModePtr = StateMachine.GetCurrentMode().Pin())
 	{
 		if (const UChaosMovementMode* ChaosMode = Cast<UChaosMovementMode>(CurrentModePtr.Get()))
@@ -361,6 +373,9 @@ void UChaosMoverSimulation::OnSimulationRollback(const FMoverTimeStep& NewTimeSt
 	// Rollback blackboard on the first frame of resimulation
 	Blackboard->Invalidate(EInvalidationReason::Rollback);
 
+	// Remove any events that have been stored for this simulation tick
+	Events.Empty();
+
 	// Make sure the movement basis matches the one in the sync state
 	if (FChaosMovementBasis* MovementBasis = CurrentSyncState.SyncStateCollection.FindMutableDataByType<FChaosMovementBasis>())
 	{
@@ -464,28 +479,23 @@ void UChaosMoverSimulation::PostSimulationTickCharacter(const IChaosCharacterMov
 
 void UChaosMoverSimulation::PostSimulationTickCharacterConstraint(const IChaosCharacterConstraintMovementModeInterface& CharacterConstraintMode, const FMoverTimeStep& TimeStep, UE::ChaosMover::FSimulationOutputData& OutputData)
 {
-	Chaos::FCharacterGroundConstraintHandle* ConstraintHandle = nullptr;
-	if (CharacterConstraintProxy && CharacterConstraintProxy->IsInitialized())
-	{
-		ConstraintHandle = CharacterConstraintProxy->GetPhysicsThreadAPI();
-	}
-	if (!ConstraintHandle)
+	if (!CharacterConstraintHandle)
 	{
 		return;
 	}
 
-	if (CharacterConstraintMode.ShouldEnableConstraint() && !ConstraintHandle->IsEnabled())
+	if (CharacterConstraintMode.ShouldEnableConstraint() && !CharacterConstraintHandle->IsEnabled())
 	{
 		EnableCharacterConstraint();
 	}
-	else if (!CharacterConstraintMode.ShouldEnableConstraint() && ConstraintHandle->IsEnabled())
+	else if (!CharacterConstraintMode.ShouldEnableConstraint() && CharacterConstraintHandle->IsEnabled())
 	{
 		DisableCharacterConstraint();
 		return;
 	}
 
 	// Update the up direction in the settings
-	Chaos::FCharacterGroundConstraintSettings& Settings = ConstraintHandle->GetSettings_Mutable();
+	Chaos::FCharacterGroundConstraintSettings& Settings = CharacterConstraintHandle->GetSettings_Mutable();
 	if (const FChaosMoverSimulationDefaultInputs* SimInputs = LocalSimInput.FindDataByType<FChaosMoverSimulationDefaultInputs>())
 	{
 		Settings.VerticalAxis = SimInputs->UpDir;
@@ -515,10 +525,10 @@ void UChaosMoverSimulation::PostSimulationTickCharacterConstraint(const IChaosCh
 				}
 			}
 		}
-		ConstraintHandle->SetGroundParticle(GroundParticle);
+		CharacterConstraintHandle->SetGroundParticle(GroundParticle);
 
 		// Set the max walkable slope angle using any override from the hit component
-		float WalkableSlopeCosine = ConstraintHandle->GetSettings().CosMaxWalkableSlopeAngle;
+		float WalkableSlopeCosine = CharacterConstraintHandle->GetSettings().CosMaxWalkableSlopeAngle;
 		if (const TStrongObjectPtr<UPrimitiveComponent> PrimComp = FloorResult.HitResult.Component.Pin())
 		{
 			const FWalkableSlopeOverride& SlopeOverride = PrimComp->GetWalkableSlopeOverride();
@@ -533,7 +543,7 @@ void UChaosMoverSimulation::PostSimulationTickCharacterConstraint(const IChaosCh
 		FChaosMoverCharacterSimState* CharacterSimState = InternalSimData.FindMutableDataByType<FChaosMoverCharacterSimState>();
 		check(CharacterSimState);
 
-		ConstraintHandle->SetData({
+		CharacterConstraintHandle->SetData({
 			FloorResult.HitResult.ImpactNormal,
 			CharacterSimState->TargetDeltaPosition,
 			CharacterSimState->TargetDeltaFacing,
@@ -543,9 +553,9 @@ void UChaosMoverSimulation::PostSimulationTickCharacterConstraint(const IChaosCh
 	}
 	else
 	{
-		ConstraintHandle->SetGroundParticle(nullptr);
-		ConstraintHandle->SetData({
-			ConstraintHandle->GetSettings().VerticalAxis,
+		CharacterConstraintHandle->SetGroundParticle(nullptr);
+		CharacterConstraintHandle->SetData({
+			CharacterConstraintHandle->GetSettings().VerticalAxis,
 			Chaos::FVec3::ZeroVector,
 			0.0,
 			1.0e10,
@@ -556,19 +566,14 @@ void UChaosMoverSimulation::PostSimulationTickCharacterConstraint(const IChaosCh
 
 void UChaosMoverSimulation::PostSimulationTickMovementActuation(const IChaosMovementActuationInterface& MovementActuationMode, const FMoverTimeStep& TimeStep, UE::ChaosMover::FSimulationOutputData& OutputData)
 {
-	Chaos::FPBDJointConstraintHandle* ConstraintHandle = nullptr;
-	if (ActuationConstraintProxy && ActuationConstraintProxy->IsInitialized())
-	{
-		ConstraintHandle = ActuationConstraintProxy->GetHandle();
-	}
-	if (!ConstraintHandle)
+	if (!ActuationConstraintHandle)
 	{
 		return;
 	}
 
 	if (MovementActuationMode.ShouldUseConstraint())
 	{
-		if (!ConstraintHandle->IsEnabled())
+		if (!ActuationConstraintHandle->IsEnabled())
 		{
 			EnableActuationConstraint();
 		}
@@ -578,11 +583,11 @@ void UChaosMoverSimulation::PostSimulationTickMovementActuation(const IChaosMove
 			SetControlledParticleDynamic();
 		}
 
-		ConstraintHandle->SetSettings(MovementActuationMode.GetConstraintSettings());
+		ActuationConstraintHandle->SetSettings(MovementActuationMode.GetConstraintSettings());
 	}
 	else if (!MovementActuationMode.ShouldUseConstraint())
 	{
-		if (ConstraintHandle->IsEnabled())
+		if (ActuationConstraintHandle->IsEnabled())
 		{
 			DisableActuationConstraint();
 		}
@@ -606,12 +611,10 @@ void UChaosMoverSimulation::SetActuationTargetTransform(const FTransform& Target
 		return;
 	}
 
-	Chaos::FPBDJointConstraintHandle* ConstraintHandle = ActuationConstraintProxy && ActuationConstraintProxy->IsInitialized() ? ActuationConstraintProxy->GetHandle() : nullptr;
-	Chaos::FKinematicGeometryParticleHandle* EndpointParticleHandle = (ConstraintHandle && ActuationConstraintKinematicEndPointProxy) ? ActuationConstraintKinematicEndPointProxy->GetHandle_LowLevel()->CastToKinematicParticle() : nullptr;
 	Chaos::FKinematicTarget KinematicPathTarget = Chaos::FKinematicTarget::MakePositionTarget(TargetTransform);
-	if (EndpointParticleHandle)
+	if (ActuationConstraintEndPointParticleHandle)
 	{	
-		Evolution->SetParticleKinematicTarget(EndpointParticleHandle, KinematicPathTarget);
+		Evolution->SetParticleKinematicTarget(ActuationConstraintEndPointParticleHandle, KinematicPathTarget);
 	}
 	if (IsControlledParticleKinematic())
 	{
@@ -630,11 +633,9 @@ void UChaosMoverSimulation::TeleportActuationTarget(const FTransform& TargetTran
 		return;
 	}
 
-	Chaos::FPBDJointConstraintHandle* ConstraintHandle = (ActuationConstraintProxy && ActuationConstraintProxy->IsInitialized()) ? ActuationConstraintProxy->GetHandle() : nullptr;
-	Chaos::FKinematicGeometryParticleHandle* EndpointParticleHandle = (ConstraintHandle) ? ActuationConstraintKinematicEndPointProxy->GetHandle_LowLevel()->CastToKinematicParticle() : nullptr;
-	if (EndpointParticleHandle)
+	if (ActuationConstraintEndPointParticleHandle)
 	{	
-		Evolution->SetParticleTransform(EndpointParticleHandle, TargetTransform.GetTranslation(), TargetTransform.GetRotation(), /*bIsTeleport=*/true);
+		Evolution->SetParticleTransform(ActuationConstraintEndPointParticleHandle, TargetTransform.GetTranslation(), TargetTransform.GetRotation(), /*bIsTeleport=*/true);
 	}
 	if (AlsoTeleportControlledParticle)
 	{
@@ -855,11 +856,9 @@ void UChaosMoverSimulation::OnMovementModeChanged(const FMovementModeChangedEven
 	{
 		if (IChaosCharacterConstraintMovementModeInterface* NextCharacterConstraintMode = Cast<IChaosCharacterConstraintMovementModeInterface>(NextModePtr.Get()))
 		{
-			bool IsCharacterConstraintInitialized = (CharacterConstraintProxy != nullptr) && CharacterConstraintProxy->IsInitialized();
-			Chaos::FCharacterGroundConstraintHandle* ConstraintHandle = IsCharacterConstraintInitialized ? CharacterConstraintProxy->GetPhysicsThreadAPI() : nullptr;
-			if (ConstraintHandle)
+			if (CharacterConstraintHandle)
 			{
-				Chaos::FCharacterGroundConstraintSettings& Settings = ConstraintHandle->GetSettings_Mutable();
+				Chaos::FCharacterGroundConstraintSettings& Settings = CharacterConstraintHandle->GetSettings_Mutable();
 				NextCharacterConstraintMode->UpdateConstraintSettings(Settings);
 
 				// Character ground constraint modes currently assume moving a dynamic particle and using a character ground constraint
@@ -880,13 +879,12 @@ void UChaosMoverSimulation::OnMovementModeChanged(const FMovementModeChangedEven
 				*PresimPathedMovementState = FChaosPathedMovementState();
 			}
 
-			Chaos::FPBDJointConstraintHandle* ConstraintHandle = (ActuationConstraintProxy && ActuationConstraintProxy->IsInitialized()) ? ActuationConstraintProxy->GetHandle() : nullptr;
-			if (ConstraintHandle)
+			if (ActuationConstraintHandle)
 			{
 				if (NextMovementActuationInterface->ShouldUseConstraint())
 				{
 					// Enable the constraint
-					if (!ConstraintHandle->IsEnabled())
+					if (!ActuationConstraintHandle->IsEnabled())
 					{
 						EnableActuationConstraint();
 					}
@@ -898,7 +896,7 @@ void UChaosMoverSimulation::OnMovementModeChanged(const FMovementModeChangedEven
 					// Apply movement actuation constraint settings
 					// Bug: if we didn't also call SetSettings every frame in UChaosMoverSimulation::PostSimulationTickMovementActuation,
 					// the settings wouldn't apply correctly from this call only if this mode is the first active mode (OnMovementModeChanged called on the very first OnSimulationTick)
-					ConstraintHandle->SetSettings(NextMovementActuationInterface->GetConstraintSettings());
+					ActuationConstraintHandle->SetSettings(NextMovementActuationInterface->GetConstraintSettings());
 
 					// If the mode was also a pathed mode, let's teleport to the position on the path
 					// we're supposed to be at, using CurrentPathProgress
@@ -914,7 +912,7 @@ void UChaosMoverSimulation::OnMovementModeChanged(const FMovementModeChangedEven
 				}
 				else
 				{
-					if (ConstraintHandle->IsEnabled())
+					if (ActuationConstraintHandle->IsEnabled())
 					{
 						DisableActuationConstraint();
 					}
@@ -981,59 +979,48 @@ bool UChaosMoverSimulation::IsControlledParticleKinematic() const
 
 void UChaosMoverSimulation::EnableCharacterConstraint()
 {
-	if (CharacterConstraintProxy && CharacterConstraintProxy->IsInitialized())
+	if (CharacterConstraintHandle)
 	{
-		Chaos::FCharacterGroundConstraintHandle* ConstraintHandle = CharacterConstraintProxy->GetPhysicsThreadAPI();
-
-		if (ConstraintHandle->GetCharacterParticle())
+		if (CharacterConstraintHandle->GetCharacterParticle())
 		{
-			ConstraintHandle->SetEnabled(true);
+			CharacterConstraintHandle->SetEnabled(true);
 		}
 	}
 }
 
 void UChaosMoverSimulation::DisableCharacterConstraint()
 {
-	if (CharacterConstraintProxy && CharacterConstraintProxy->IsInitialized())
+	if (CharacterConstraintHandle)
 	{
-		Chaos::FCharacterGroundConstraintHandle* ConstraintHandle = CharacterConstraintProxy->GetPhysicsThreadAPI();
-
-		ConstraintHandle->SetEnabled(false);
+		CharacterConstraintHandle->SetEnabled(false);
 	}
 }
 
 bool UChaosMoverSimulation::IsCharacterConstraintEnabled() const
 {
-	const Chaos::FCharacterGroundConstraintHandle* ConstraintHandle = (CharacterConstraintProxy && CharacterConstraintProxy->IsInitialized()) ? CharacterConstraintProxy->GetPhysicsThreadAPI() : nullptr;
-	return (ConstraintHandle && ConstraintHandle->IsEnabled());
+	return (CharacterConstraintHandle && CharacterConstraintHandle->IsEnabled());
 }
 
 void UChaosMoverSimulation::EnableActuationConstraint()
 {
-	if (ActuationConstraintProxy && ActuationConstraintProxy->IsInitialized())
+	if (ActuationConstraintHandle)
 	{
-		if (Chaos::FPBDJointConstraintHandle* ConstraintHandle = ActuationConstraintProxy->GetHandle())
-		{
-			ConstraintHandle->SetConstraintEnabled(true);
-		}
+		ActuationConstraintHandle->SetConstraintEnabled(true);
 	}
 }
 
 void UChaosMoverSimulation::DisableActuationConstraint()
 {
-	if (ActuationConstraintProxy && ActuationConstraintProxy->IsInitialized())
+
+	if (ActuationConstraintHandle)
 	{
-		if (Chaos::FPBDJointConstraintHandle* ConstraintHandle = ActuationConstraintProxy->GetHandle())
-		{
-			ConstraintHandle->SetConstraintEnabled(false);
-		}
+		ActuationConstraintHandle->SetConstraintEnabled(false);
 	}
 }
 
 bool UChaosMoverSimulation::IsActuationConstraintEnabled() const
 {
-	const Chaos::FPBDJointConstraintHandle* ConstraintHandle = (ActuationConstraintProxy && ActuationConstraintProxy->IsInitialized()) ? ActuationConstraintProxy->GetHandle() : nullptr;
-	return (ConstraintHandle && ConstraintHandle->IsEnabled());
+	return (ActuationConstraintHandle && ActuationConstraintHandle->IsEnabled());
 }
 
 void UChaosMoverSimulation::K2_QueueInstantMovementEffect(const int32& EffectAsRawData)
diff --git a/Engine/Plugins/Experimental/ChaosMover/Source/ChaosMover/Public/ChaosMover/Backends/ChaosMoverBackend.h b/Engine/Plugins/Experimental/ChaosMover/Source/ChaosMover/Public/ChaosMover/Backends/ChaosMoverBackend.h
--- a/Engine/Plugins/Experimental/ChaosMover/Source/ChaosMover/Public/ChaosMover/Backends/ChaosMoverBackend.h
+++ b/Engine/Plugins/Experimental/ChaosMover/Source/ChaosMover/Public/ChaosMover/Backends/ChaosMoverBackend.h
@@ -103,7 +103,7 @@ protected:
 
 	bool bIsUsingAsyncPhysics = false;
 	bool bWantsDestroySim = false;
-	bool bWantsCreateSim = false;
+	bool bWantsCreateSim = true;
 
 	// Transfers queued instant movement effects from the mover component to the simulation input.
 	// This is used in the input producing case.
diff --git a/Engine/Plugins/Experimental/ChaosMover/Source/ChaosMover/Public/ChaosMover/ChaosMoverSimulation.h b/Engine/Plugins/Experimental/ChaosMover/Source/ChaosMover/Public/ChaosMover/ChaosMoverSimulation.h
--- a/Engine/Plugins/Experimental/ChaosMover/Source/ChaosMover/Public/ChaosMover/ChaosMoverSimulation.h
+++ b/Engine/Plugins/Experimental/ChaosMover/Source/ChaosMover/Public/ChaosMover/ChaosMoverSimulation.h
@@ -10,15 +10,15 @@
 #include "MoveLibrary/MoverBlackboard.h"
 #include "MoverTypes.h"
 #include "MoverSimulation.h"
-#include "UObject/WeakObjectPtrTemplates.h"
 #include "PhysicsEngine/ConstraintInstance.h"
+#include "UObject/WeakObjectPtrTemplates.h"
 
 #include "ChaosMoverSimulation.generated.h"
 
 namespace Chaos
 {
-	class FCharacterGroundConstraintProxy;
-	class FJointConstraintPhysicsProxy;
+	class FCharacterGroundConstraintHandle;
+	class FPBDJointConstraintHandle;
 	class FCollisionContactModifier;
 }
 
@@ -88,9 +88,9 @@ public:
 		TWeakObjectPtr<UImmediateMovementModeTransition> ImmediateModeTransition = nullptr;
 		TWeakObjectPtr<UMovementMixer> MovementMixer = nullptr;
 		FTransform TransformOnInit = FTransform::Identity;
-		Chaos::FCharacterGroundConstraintProxy* CharacterConstraintProxy = nullptr;
-		Chaos::FJointConstraintPhysicsProxy* ActuationConstraintProxy = nullptr;
-		Chaos::FSingleParticlePhysicsProxy* ActuationConstraintKinematicEndPointProxy = nullptr;
+		Chaos::FCharacterGroundConstraintHandle* CharacterConstraintHandle = nullptr;
+		Chaos::FPBDJointConstraintHandle* ActuationConstraintHandle = nullptr;
+		Chaos::FKinematicGeometryParticleHandle* ActuationConstraintEndPointParticleHandle = nullptr;
 		Chaos::FConstPhysicsObjectHandle PhysicsObject = nullptr;
 		Chaos::FPhysicsSolver* Solver = nullptr;
 		UWorld* World = nullptr;
@@ -194,11 +194,11 @@ protected:
 	Chaos::FConstPhysicsObjectHandle PhysicsObject = nullptr;
 
 	// Character ground constraint, specifically for moving on ground like characters
-	Chaos::FCharacterGroundConstraintProxy* CharacterConstraintProxy = nullptr;
+	Chaos::FCharacterGroundConstraintHandle* CharacterConstraintHandle = nullptr;
 
 	// General purpose actuation joint constraint
-	Chaos::FJointConstraintPhysicsProxy* ActuationConstraintProxy = nullptr;
-	Chaos::FSingleParticlePhysicsProxy* ActuationConstraintKinematicEndPointProxy = nullptr;
+	Chaos::FPBDJointConstraintHandle* ActuationConstraintHandle = nullptr;
+	Chaos::FKinematicGeometryParticleHandle* ActuationConstraintEndPointParticleHandle= nullptr;
 	
 	// Some movement modes are relative to a movement basis transform, stored here
 	FTransform MovementBasisTransform = FTransform::Identity;
@@ -215,4 +215,6 @@ private:
 	bool bSyncStateOverridden = false;
 
 	int32 InternalServerFrame = 0;
+
+	bool bInitialized = false;
 };
