ChaosMover - Fix potential issue if async dt is zero

#rnx

diff --git a/Engine/Plugins/Experimental/ChaosMover/Source/ChaosMover/Private/Backends/ChaosMoverSubsystem.cpp b/Engine/Plugins/Experimental/ChaosMover/Source/ChaosMover/Private/Backends/ChaosMoverSubsystem.cpp
--- a/Engine/Plugins/Experimental/ChaosMover/Source/ChaosMover/Private/Backends/ChaosMoverSubsystem.cpp
+++ b/Engine/Plugins/Experimental/ChaosMover/Source/ChaosMover/Private/Backends/ChaosMoverSubsystem.cpp
@@ -119,8 +119,17 @@ void UChaosMoverSubsystem::InjectInputs_External(int32 PhysicsStep, int32 NumSte
 				
 				if (Solver->IsUsingAsyncResults())
 				{
-					TimeStep.ServerFrame = FMath::FloorToInt(Solver->GetPhysicsResultsTime_External() / (Solver->GetAsyncDeltaTime()));
-					TimeStep.StepMs = Solver->GetAsyncDeltaTime() * 1000.0f;
+					const float AsyncDt = Solver->GetAsyncDeltaTime();
+					if (AsyncDt > UE_SMALL_NUMBER)
+					{
+						TimeStep.ServerFrame = FMath::FloorToInt(Solver->GetPhysicsResultsTime_External() / AsyncDt);
+					}
+					else
+					{
+						TimeStep.ServerFrame = Solver->GetCurrentFrame();
+					}
+					
+					TimeStep.StepMs = AsyncDt * 1000.0f;
 				}
 				else
 				{
