[Mover] Add Game Thread context to FMoverSimulationEventData::EventProcessedCallback. By default, the context contains the MoverComponent pointer.
[ChaosMover] Changes to make crouching work in ChaosMover

diff --git a/Engine/Plugins/Experimental/ChaosMover/Source/ChaosMover/Private/ChaosMoverSimulation.cpp b/Engine/Plugins/Experimental/ChaosMover/Source/ChaosMover/Private/ChaosMoverSimulation.cpp
--- a/Engine/Plugins/Experimental/ChaosMover/Source/ChaosMover/Private/ChaosMoverSimulation.cpp
+++ b/Engine/Plugins/Experimental/ChaosMover/Source/ChaosMover/Private/ChaosMoverSimulation.cpp
@@ -1139,6 +1139,17 @@ bool UChaosMoverSimulation::HasGameplayTag(FGameplayTag TagToFind, bool bExactMa
 
 const FMovementModifierBase* UChaosMoverSimulation::FindMovementModifierByType(const UScriptStruct* DataStructType) const
 {
+	// Check active modifiers for modifier handle
+	for (auto ActiveModifierFromSyncStateIt = CurrentSyncState.MovementModifiers.GetActiveModifiersIterator(); ActiveModifierFromSyncStateIt; ++ActiveModifierFromSyncStateIt)
+	{
+		const TSharedPtr<FMovementModifierBase> ActiveModifierFromSyncState = *ActiveModifierFromSyncStateIt;
+
+		if (DataStructType == ActiveModifierFromSyncState->GetScriptStruct())
+		{
+			return ActiveModifierFromSyncState.Get();
+		}
+	}
+
 	return StateMachine.FindQueuedModifierByType(DataStructType);
 }
 
diff --git a/Engine/Plugins/Experimental/ChaosMover/Source/ChaosMover/Private/Character/ChaosCharacterMoverComponent.cpp b/Engine/Plugins/Experimental/ChaosMover/Source/ChaosMover/Private/Character/ChaosCharacterMoverComponent.cpp
--- a/Engine/Plugins/Experimental/ChaosMover/Source/ChaosMover/Private/Character/ChaosCharacterMoverComponent.cpp
+++ b/Engine/Plugins/Experimental/ChaosMover/Source/ChaosMover/Private/Character/ChaosCharacterMoverComponent.cpp
@@ -117,16 +117,11 @@ void UChaosCharacterMoverComponent::ProduceInput(const int32 DeltaTimeMS, FMover
 		ClearQueuedMode();
 	}
 
-	if (bCrouchInputPending)
-	{
-		if ((bWantsToCrouch && CanCrouch()) || !bWantsToCrouch)
-		{
-			FChaosMoverCrouchInputs& CrouchInputs = Cmd->InputCollection.FindOrAddMutableDataByType<FChaosMoverCrouchInputs>();
-			CrouchInputs.bWantsToCrouch = bWantsToCrouch;
-		}
-
-		bCrouchInputPending = false;
-	}
+	// TODO @Harsha Making bWantsToCrouch sticky because the server is missing the one time crouch input.
+	// Investigate the input drops and after the fix, revert this change back to making it non-sticky.
+	bWantsToCrouch &= CanCrouch();
+	FChaosMoverCrouchInputs& CrouchInputs = Cmd->InputCollection.FindOrAddMutableDataByType<FChaosMoverCrouchInputs>();
+	CrouchInputs.bWantsToCrouch = bWantsToCrouch;
 
 	if (bCancelMovementOverrides)
 	{
diff --git a/Engine/Plugins/Experimental/ChaosMover/Source/ChaosMover/Private/Character/Modifiers/ChaosStanceModifier.cpp b/Engine/Plugins/Experimental/ChaosMover/Source/ChaosMover/Private/Character/Modifiers/ChaosStanceModifier.cpp
--- a/Engine/Plugins/Experimental/ChaosMover/Source/ChaosMover/Private/Character/Modifiers/ChaosStanceModifier.cpp
+++ b/Engine/Plugins/Experimental/ChaosMover/Source/ChaosMover/Private/Character/Modifiers/ChaosStanceModifier.cpp
@@ -6,6 +6,7 @@
 #include "Chaos/PhysicsObjectInternalInterface.h"
 #include "ChaosMover/ChaosMoverLog.h"
 #include "ChaosMover/ChaosMoverSimulation.h"
+#include "DefaultMovementSet/CharacterMoverComponent.h"
 #include "Components/CapsuleComponent.h"
 #include "MoverTypes.h"
 #include "PBDRigidsSolver.h"
@@ -166,7 +167,7 @@ void FChaosStanceModifier::UpdateStance(const FMovementModifierParams_Async& Asy
 	Chaos::FImplicitObjectPtr NewCapsule = ParticleHandle->GetGeometry();
 	Chaos::FPhysicsObjectHandle PhysicsObject = SimInputs->PhysicsObject;
 
-	auto SetGeometryGTCallback = [NewCapsule, PhysicsObject](const FMoverSimulationEventData & Data) {
+	auto SetGeometryGTCallback = [NewCapsule, PhysicsObject](const FMoverSimulationEventData & Data, const FMoverSimEventGameThreadContext& GTContext) {
 		Chaos::FPBDRigidParticle* Particle = FPhysicsObjectExternalInterface::LockRead(PhysicsObject)->GetRigidParticle(PhysicsObject);
 		if (Particle && NewCapsule.IsValid())
 		{
@@ -181,6 +182,19 @@ void FChaosStanceModifier::UpdateStance(const FMovementModifierParams_Async& Asy
 void FChaosStanceModifier::OnStart_Async(const FMovementModifierParams_Async& AsyncParams)
 {
 	UpdateStance(AsyncParams, EStanceMode::Crouch);
+	
+	// Propagate crouch state to GT if it's triggered on the PT (rollback or resim)
+	// TODO @Harsha this is not necessary if bWantsToCrouch is non-sticky, stale bWantsToCrouch will not matter then
+	if (UChaosMoverSimulation* Sim = Cast<UChaosMoverSimulation>(AsyncParams.Simulation))
+	{
+		auto EnableWantsToCrouch = [] (const FMoverSimulationEventData& Data, const FMoverSimEventGameThreadContext& GTContext) {
+			if (UCharacterMoverComponent* CharMoverComp = Cast<UCharacterMoverComponent>(GTContext.MoverComp))
+			{
+				CharMoverComp->Crouch();
+			}
+		};
+		Sim->AddEvent(MakeShared<FStanceModifierModeChangedEventData>(AsyncParams.TimeStep->BaseSimTimeMs, CurrentModeName, AsyncParams.SyncState->MovementMode, EnableWantsToCrouch));
+	}
 
 	// Set the current mode - this is used to check if the mode changes
 	CurrentModeName = AsyncParams.SyncState->MovementMode;
@@ -189,6 +203,19 @@ void FChaosStanceModifier::OnStart_Async(const FMovementModifierParams_Async& As
 void FChaosStanceModifier::OnEnd_Async(const FMovementModifierParams_Async& AsyncParams)
 {
 	UpdateStance(AsyncParams, EStanceMode::Invalid);
+	
+	// Propagate uncrouch state to GT if it's triggered on the PT (rollback, resim or uncrouching on mode change)
+	// TODO @Harsha this is not necessary if bWantsToCrouch is non-sticky, stale bWantsToCrouch will not matter then
+	if (UChaosMoverSimulation* Sim = Cast<UChaosMoverSimulation>(AsyncParams.Simulation))
+	{
+		auto DisableWantsToCrouch = [](const FMoverSimulationEventData& Data, const FMoverSimEventGameThreadContext& GTContext) {
+			if (UCharacterMoverComponent* CharMoverComp = Cast<UCharacterMoverComponent>(GTContext.MoverComp))
+			{
+				CharMoverComp->UnCrouch();
+			}
+		};
+		Sim->AddEvent(MakeShared<FStanceModifierModeChangedEventData>(AsyncParams.TimeStep->BaseSimTimeMs, CurrentModeName, AsyncParams.SyncState->MovementMode, DisableWantsToCrouch));
+	}
 
 	CurrentModeName = NAME_None;
 }
diff --git a/Engine/Plugins/Experimental/ChaosMover/Source/ChaosMover/Public/ChaosMover/ChaosMoverSimulationTypes.h b/Engine/Plugins/Experimental/ChaosMover/Source/ChaosMover/Public/ChaosMover/ChaosMoverSimulationTypes.h
--- a/Engine/Plugins/Experimental/ChaosMover/Source/ChaosMover/Public/ChaosMover/ChaosMoverSimulationTypes.h
+++ b/Engine/Plugins/Experimental/ChaosMover/Source/ChaosMover/Public/ChaosMover/ChaosMoverSimulationTypes.h
@@ -295,6 +295,28 @@ struct FStanceModifiedEventData : public FMoverSimulationEventData
 	EStanceMode NewStance = EStanceMode::Invalid;
 };
 
+USTRUCT()
+struct FStanceModifierModeChangedEventData : public FMoverSimulationEventData
+{
+	GENERATED_BODY()
+	
+	FStanceModifierModeChangedEventData(double InEventTimeMs, FName InOldMode, FName InNewMode, FEventProcessedCallbackPtr InEventProcessedCallback = nullptr)
+	: FMoverSimulationEventData(InEventTimeMs, InEventProcessedCallback)
+	, OldMode(InOldMode)
+	, NewMode(InNewMode)
+	{
+	}
+	FStanceModifierModeChangedEventData() {}
+
+	virtual UScriptStruct* GetScriptStruct() const override
+	{
+		return FStanceModifierModeChangedEventData::StaticStruct();
+	}
+
+	FName OldMode = NAME_None;
+	FName NewMode = NAME_None;
+};
+
 // Version of a FScheduledInstantMovementEffect with Issuance Server Frame
 USTRUCT()
 struct FChaosScheduledInstantMovementEffect
diff --git a/Engine/Plugins/Experimental/Mover/Source/Mover/Private/MoverComponent.cpp b/Engine/Plugins/Experimental/Mover/Source/Mover/Private/MoverComponent.cpp
--- a/Engine/Plugins/Experimental/Mover/Source/Mover/Private/MoverComponent.cpp
+++ b/Engine/Plugins/Experimental/Mover/Source/Mover/Private/MoverComponent.cpp
@@ -2392,7 +2392,8 @@ void UMoverComponent::SetSimulationOutput(const FMoverTimeStep& TimeStep, const
 void UMoverComponent::DispatchSimulationEvent(const FMoverSimulationEventData& EventData)
 {
 	// This gives the event a callback when it is processed on the game thread
-	EventData.OnEventProcessed();
+	FMoverSimEventGameThreadContext GTContext({ this });
+	EventData.OnEventProcessed(GTContext);
 
 	// Process the simulation event at the mover component (or derived) level
 	ProcessSimulationEvent(EventData);
diff --git a/Engine/Plugins/Experimental/Mover/Source/Mover/Public/MoverSimulationTypes.h b/Engine/Plugins/Experimental/Mover/Source/Mover/Public/MoverSimulationTypes.h
--- a/Engine/Plugins/Experimental/Mover/Source/Mover/Public/MoverSimulationTypes.h
+++ b/Engine/Plugins/Experimental/Mover/Source/Mover/Public/MoverSimulationTypes.h
@@ -487,12 +487,21 @@ struct FMoverPredictTrajectoryParams
 
 };
 
+USTRUCT()
+struct FMoverSimEventGameThreadContext
+{
+	GENERATED_BODY()
+
+public:
+	UMoverComponent* MoverComp = nullptr;
+};
+
 USTRUCT()
 struct FMoverSimulationEventData
 {
 	GENERATED_BODY()
 
-	using FEventProcessedCallbackPtr = std::function<void(const FMoverSimulationEventData& Data)>;
+	using FEventProcessedCallbackPtr = std::function<void(const FMoverSimulationEventData& Data, const FMoverSimEventGameThreadContext& GameThreadContext)>;
 
 	FMoverSimulationEventData(double InEventTimeMs, FEventProcessedCallbackPtr InEventProcessedCallback = nullptr)
 		: EventProcessedCallback(InEventProcessedCallback)
@@ -517,11 +526,11 @@ struct FMoverSimulationEventData
 		return const_cast<const T*>(const_cast<FMoverSimulationEventData*>(this)->CastTo_Mutable<T>());
 	}
 
-	void OnEventProcessed() const
+	void OnEventProcessed(const FMoverSimEventGameThreadContext& GameThreadContext) const
 	{
 		if (EventProcessedCallback)
 		{
-			EventProcessedCallback(*this);
+			EventProcessedCallback(*this, GameThreadContext);
 		}
 	}
 
