ChaosMover - Update floor each frame

#rnx

diff --git a/Engine/Plugins/Experimental/ChaosMover/Source/ChaosMover/Private/ChaosMoverSimulation.cpp b/Engine/Plugins/Experimental/ChaosMover/Source/ChaosMover/Private/ChaosMoverSimulation.cpp
--- a/Engine/Plugins/Experimental/ChaosMover/Source/ChaosMover/Private/ChaosMoverSimulation.cpp
+++ b/Engine/Plugins/Experimental/ChaosMover/Source/ChaosMover/Private/ChaosMoverSimulation.cpp
@@ -18,6 +18,7 @@
 #include "ChaosMover/PathedMovement/ChaosPathedMovementTypes.h"
 #include "ChaosMover/Utilities/ChaosGroundMovementUtils.h"
 #include "ChaosMover/Utilities/ChaosMoverQueryUtils.h"
+#include "DefaultMovementSet/CharacterMoverSimulationTypes.h"
 #include "Engine/World.h"
 #include "Framework/Threading.h"
 #include "MoveLibrary/FloorQueryUtils.h"
@@ -484,6 +485,9 @@ void UChaosMoverSimulation::PostSimulationTickCharacter(const IChaosCharacterMov
 	// Note: We want to record the movement base but we don't record the transform
 	// so don't use this to get a relative transform for the sync state
 	PostSimDefaultSyncState->SetMovementBase(FoundFloor ? FloorResult.HitResult.GetComponent() : nullptr);
+
+	FFloorResultData& FloorData = OutputData.AdditionalOutputData.FindOrAddMutableDataByType<FFloorResultData>();
+	FloorData.FloorResult = FloorResult;
 }
 
 void UChaosMoverSimulation::PostSimulationTickCharacterConstraint(const IChaosCharacterConstraintMovementModeInterface& CharacterConstraintMode, const FMoverTimeStep& TimeStep, UE::ChaosMover::FSimulationOutputData& OutputData)
