PBCM: Hookup owner name and role in chaos state machine

#rnx

diff --git a/Engine/Plugins/Experimental/ChaosMover/Source/ChaosMover/Private/Backends/ChaosMoverBackend.cpp b/Engine/Plugins/Experimental/ChaosMover/Source/ChaosMover/Private/Backends/ChaosMoverBackend.cpp
--- a/Engine/Plugins/Experimental/ChaosMover/Source/ChaosMover/Private/Backends/ChaosMoverBackend.cpp
+++ b/Engine/Plugins/Experimental/ChaosMover/Source/ChaosMover/Private/Backends/ChaosMoverBackend.cpp
@@ -63,6 +63,7 @@ void UChaosMoverBackendComponent::InitializeComponent()
 	if (UWorld* World = GetWorld(); World && World->IsGameWorld())
 	{
 		Simulation = NewObject<UChaosMoverSimulation>(GetOwner(), SimulationClass, TEXT("ChaosMoverSimulation"), RF_Transient);
+		Simulation->SetOwner(GetMoverComponent().GetOwner());
 
 		NullMovementMode = NewObject<UNullMovementMode>(&GetMoverComponent(), TEXT("NullMovementMode"), RF_Transient);
 		ImmediateModeTransition = NewObject<UImmediateMovementModeTransition>(&GetMoverComponent(), TEXT("ImmediateModeTransition"), RF_Transient);
diff --git a/Engine/Plugins/Experimental/ChaosMover/Source/ChaosMover/Private/ChaosMoverSimulation.cpp b/Engine/Plugins/Experimental/ChaosMover/Source/ChaosMover/Private/ChaosMoverSimulation.cpp
--- a/Engine/Plugins/Experimental/ChaosMover/Source/ChaosMover/Private/ChaosMoverSimulation.cpp
+++ b/Engine/Plugins/Experimental/ChaosMover/Source/ChaosMover/Private/ChaosMoverSimulation.cpp
@@ -161,6 +161,16 @@ void UChaosMoverSimulation::Init(const FInitParams& InitParams)
 	bInitialized = true;
 }
 
+void UChaosMoverSimulation::SetOwner(const AActor* Owner)
+{
+	if (ensure(Owner))
+	{
+		OwnerLocalRole = Owner->GetLocalRole();
+		StateMachine.SetOwnerActorName(Owner->GetName());
+		StateMachine.SetOwnerActorLocalNetRole(OwnerLocalRole);
+	}
+}
+
 void UChaosMoverSimulation::Deinit()
 {
 	bInitialized = false;
diff --git a/Engine/Plugins/Experimental/ChaosMover/Source/ChaosMover/Public/ChaosMover/ChaosMoverSimulation.h b/Engine/Plugins/Experimental/ChaosMover/Source/ChaosMover/Public/ChaosMover/ChaosMoverSimulation.h
--- a/Engine/Plugins/Experimental/ChaosMover/Source/ChaosMover/Public/ChaosMover/ChaosMoverSimulation.h
+++ b/Engine/Plugins/Experimental/ChaosMover/Source/ChaosMover/Public/ChaosMover/ChaosMoverSimulation.h
@@ -78,6 +78,8 @@ public:
 	UFUNCTION(BlueprintPure, Category = Mover, meta = (Keywords = "HasTag"))
 	CHAOSMOVER_API bool HasGameplayTag(FGameplayTag TagToFind, bool bExactMatch) const;
 
+	ENetRole GetLocalRole() const { return OwnerLocalRole; }
+
 	struct FInitParams
 	{
 		TMap<FName, TWeakObjectPtr<UBaseMovementMode>> ModesToRegister;
@@ -97,6 +99,7 @@ public:
 	};
 
 	CHAOSMOVER_API void Init(const FInitParams& InitParams);
+	CHAOSMOVER_API void SetOwner(const AActor* Owner);
 	CHAOSMOVER_API void Deinit();
 
 	CHAOSMOVER_API void ProcessInputs(int32 PhysicsStep, const FMoverTimeStep& TimeStep, const UE::ChaosMover::FSimulationInputData& InputData);
@@ -217,4 +220,5 @@ private:
 	int32 InternalServerFrame = 0;
 
 	bool bInitialized = false;
+	ENetRole OwnerLocalRole = ENetRole::ROLE_None;
 };
