Chaos Mover - Add requested inputs

#rnx

diff --git a/Engine/Plugins/Experimental/ChaosMover/Source/ChaosMover/Private/Character/ChaosCharacterInputs.cpp b/Engine/Plugins/Experimental/ChaosMover/Source/ChaosMover/Private/Character/ChaosCharacterInputs.cpp
--- a/Engine/Plugins/Experimental/ChaosMover/Source/ChaosMover/Private/Character/ChaosCharacterInputs.cpp
+++ b/Engine/Plugins/Experimental/ChaosMover/Source/ChaosMover/Private/Character/ChaosCharacterInputs.cpp
@@ -169,3 +169,43 @@ void FChaosMovementSettingsOverridesRemover::Interpolate(const FMoverDataStructB
 void FChaosMovementSettingsOverridesRemover::Merge(const FMoverDataStructBase& From)
 {
 }
+
+bool FChaosMoverRequestedMoveInputs::NetSerialize(FArchive& Ar, UPackageMap* Map, bool& bOutSuccess)
+{
+	Super::NetSerialize(Ar, Map, bOutSuccess);
+
+	Ar << RequestedVelocity;
+	Ar << bForceMaxSpeed;
+	Ar << bUseAcceleration;
+
+	bOutSuccess = true;
+	return true;
+}
+
+void FChaosMoverRequestedMoveInputs::ToString(FAnsiStringBuilderBase& Out) const
+{
+	Super::ToString(Out);
+
+	Out.Appendf("RequestedVelocity: X=%.2f Y=%.2f Z=%.2f |", RequestedVelocity.X, RequestedVelocity.Y, RequestedVelocity.Z);
+	Out.Appendf("bForceMaxSpeed: %u |", bForceMaxSpeed);
+	Out.Appendf("bUseAcceleration: %u\n", bUseAcceleration);
+}
+
+bool FChaosMoverRequestedMoveInputs::ShouldReconcile(const FMoverDataStructBase& AuthorityState) const
+{
+	const FChaosMoverRequestedMoveInputs& TypedAuthority = static_cast<const FChaosMoverRequestedMoveInputs&>(AuthorityState);
+	return RequestedVelocity != TypedAuthority.RequestedVelocity;
+}
+
+void FChaosMoverRequestedMoveInputs::Interpolate(const FMoverDataStructBase& From, const FMoverDataStructBase& To, float Pct)
+{
+	const FChaosMoverRequestedMoveInputs& TypedFrom = static_cast<const FChaosMoverRequestedMoveInputs&>(From);
+	const FChaosMoverRequestedMoveInputs& TypedTo = static_cast<const FChaosMoverRequestedMoveInputs&>(To);
+	bForceMaxSpeed = Pct < 0.5f ? TypedFrom.bForceMaxSpeed : TypedTo.bForceMaxSpeed;
+	bUseAcceleration = Pct < 0.5f ? TypedFrom.bUseAcceleration : TypedTo.bUseAcceleration;
+	RequestedVelocity = FMath::Lerp(TypedFrom.RequestedVelocity, TypedTo.RequestedVelocity, Pct);
+}
+
+void FChaosMoverRequestedMoveInputs::Merge(const FMoverDataStructBase& From)
+{
+}
diff --git a/Engine/Plugins/Experimental/ChaosMover/Source/ChaosMover/Private/Character/Modes/ChaosFallingMode.cpp b/Engine/Plugins/Experimental/ChaosMover/Source/ChaosMover/Private/Character/Modes/ChaosFallingMode.cpp
--- a/Engine/Plugins/Experimental/ChaosMover/Source/ChaosMover/Private/Character/Modes/ChaosFallingMode.cpp
+++ b/Engine/Plugins/Experimental/ChaosMover/Source/ChaosMover/Private/Character/Modes/ChaosFallingMode.cpp
@@ -4,6 +4,7 @@
 
 #include "ChaosMover/ChaosMoverLog.h"
 #include "ChaosMover/ChaosMoverSimulation.h"
+#include "ChaosMover/Character/ChaosCharacterInputs.h"
 #include "ChaosMover/Character/Settings/SharedChaosCharacterMovementSettings.h"
 #include "ChaosMover/Character/Transitions/ChaosCharacterLandingCheck.h"
 #include "ChaosMover/Utilities/ChaosMoverQueryUtils.h"
@@ -66,9 +67,20 @@ void UChaosFallingMode::GenerateMove_Implementation(const FMoverTickStartData& S
 	const FVector StartHorizontalVelocity = FVector::VectorPlaneProject(StartVelocity, DefaultSimInputs->UpDir);
 
 	FFreeMoveParams Params;
-	Params.MoveInputType = CharacterInputs->GetMoveInputType();
-	const bool bMaintainInputMagnitude = true;
-	Params.MoveInput = UPlanarConstraintUtils::ConstrainDirectionToPlane(FPlanarConstraint(), CharacterInputs->GetMoveInput_WorldSpace(), bMaintainInputMagnitude);
+
+	const FChaosMoverRequestedMoveInputs* RequestedMove = StartState.InputCmd.InputCollection.FindDataByType<FChaosMoverRequestedMoveInputs>();
+	if (RequestedMove)
+	{
+		Params.MoveInputType = EMoveInputType::Velocity;
+		Params.MoveInput = RequestedMove->bForceMaxSpeed ? RequestedMove->RequestedVelocity : RequestedMove->RequestedVelocity.GetClampedToMaxSize(GetMaxSpeed());
+	}
+	else
+	{
+		Params.MoveInputType = CharacterInputs->GetMoveInputType();
+		constexpr bool bMaintainInputMagnitude = true;
+		Params.MoveInput = UPlanarConstraintUtils::ConstrainDirectionToPlane(FPlanarConstraint(), CharacterInputs->GetMoveInput_WorldSpace(), bMaintainInputMagnitude);
+	}
+
 	Params.MoveInput *= AirControlPercentage;
 	// Don't care about up axis input since falling - if up input matters that should probably be a different movement mode
 	Params.MoveInput = FVector::VectorPlaneProject(Params.MoveInput, DefaultSimInputs->UpDir);
diff --git a/Engine/Plugins/Experimental/ChaosMover/Source/ChaosMover/Private/Character/Modes/ChaosWalkingMode.cpp b/Engine/Plugins/Experimental/ChaosMover/Source/ChaosMover/Private/Character/Modes/ChaosWalkingMode.cpp
--- a/Engine/Plugins/Experimental/ChaosMover/Source/ChaosMover/Private/Character/Modes/ChaosWalkingMode.cpp
+++ b/Engine/Plugins/Experimental/ChaosMover/Source/ChaosMover/Private/Character/Modes/ChaosWalkingMode.cpp
@@ -8,6 +8,7 @@
 #include "ChaosMover/ChaosMoverLog.h"
 #include "ChaosMover/ChaosMoverSimulation.h"
 #include "ChaosMover/ChaosMoverSimulationTypes.h"
+#include "ChaosMover/Character/ChaosCharacterInputs.h"
 #include "ChaosMover/Character/Settings/SharedChaosCharacterMovementSettings.h"
 #include "ChaosMover/Character/Transitions/ChaosCharacterFallingCheck.h"
 #include "ChaosMover/Character/Transitions/ChaosCharacterJumpCheck.h"
@@ -90,50 +91,73 @@ void UChaosWalkingMode::GenerateMove_Implementation(const FMoverTickStartData& S
 
 	const FQuat WorldToGravityTransform = FQuat::FindBetweenNormals(FVector::UpVector, DefaultSimInputs->UpDir);
 	IntendedOrientation_WorldSpace = UMovementUtils::ApplyGravityToOrientationIntent(IntendedOrientation_WorldSpace, WorldToGravityTransform, bShouldCharacterRemainUpright);
+	const FVector CurrentVelocity = StartingSyncState->GetVelocity_WorldSpace();
 
-	FGroundMoveParams Params;
-
-	if (CharacterInputs)
+	const FChaosMoverRequestedMoveInputs* RequestedMove = StartState.InputCmd.InputCollection.FindDataByType<FChaosMoverRequestedMoveInputs>();
+	if (RequestedMove)
 	{
-		Params.MoveInputType = CharacterInputs->GetMoveInputType();
-
-		const bool bMaintainInputMagnitude = true;
-		Params.MoveInput = CharacterInputs->GetMoveInput_WorldSpace();
-		//Params.MoveInput = UPlanarConstraintUtils::ConstrainDirectionToPlane(MoverComp->GetPlanarConstraint(), CharacterInputs->GetMoveInput_WorldSpace(), bMaintainInputMagnitude);
+		FRequestedMoveParams Params;
+
+		Params.PriorVelocity = FVector::VectorPlaneProject(CurrentVelocity, MovementNormal);
+		Params.PriorOrientation = StartingSyncState->GetOrientation_WorldSpace();
+		Params.GroundNormal = MovementNormal;
+		Params.TurningRate = SharedSettings->TurningRate;
+		Params.MaxSpeed = GetMaxSpeed();
+		Params.Acceleration = GetAcceleration();
+		Params.Deceleration = SharedSettings->Deceleration;
+		Params.DeltaSeconds = DeltaSeconds;
+		Params.WorldToGravityQuat = WorldToGravityTransform;
+
+		Params.bRequestedMoveWithMaxSpeed = RequestedMove->bForceMaxSpeed;
+		Params.bShouldComputeAcceleration = RequestedMove->bUseAcceleration;
+		Params.RequestedVelocity = RequestedMove->RequestedVelocity;
+
+		OutProposedMove = UChaosGroundMovementUtils::ComputeRequestedMove(Params);
 	}
 	else
 	{
-		Params.MoveInputType = EMoveInputType::None;
-		Params.MoveInput = FVector::ZeroVector;
-	}
+		FGroundMoveParams Params;
 
-	const FVector CurrentVelocity = StartingSyncState->GetVelocity_WorldSpace();
+		if (CharacterInputs)
+		{
+			Params.MoveInputType = CharacterInputs->GetMoveInputType();
 
-	Params.OrientationIntent = IntendedOrientation_WorldSpace;
-	Params.PriorVelocity = FVector::VectorPlaneProject(CurrentVelocity, MovementNormal);
-	Params.PriorOrientation = StartingSyncState->GetOrientation_WorldSpace();
-	Params.GroundNormal = MovementNormal;
-	Params.TurningRate = SharedSettings->TurningRate;
-	Params.TurningBoost = SharedSettings->TurningBoost;
-	Params.MaxSpeed = GetMaxSpeed();
-	Params.Acceleration = GetAcceleration();
-	Params.Deceleration = SharedSettings->Deceleration;
-	Params.DeltaSeconds = DeltaSeconds;
-	Params.WorldToGravityQuat = WorldToGravityTransform;
-	Params.UpDirection = UpDirection;
-	Params.bUseAccelerationForVelocityMove = SharedSettings->bUseAccelerationForVelocityMove;
-
-	if (Params.MoveInput.SizeSquared() > 0.f && !UMovementUtils::IsExceedingMaxSpeed(Params.PriorVelocity, SharedSettings->MaxSpeed))
-	{
-		Params.Friction = SharedSettings->GroundFriction;
-	}
-	else
-	{
-		Params.Friction = SharedSettings->bUseSeparateBrakingFriction ? SharedSettings->BrakingFriction : SharedSettings->GroundFriction;
-		Params.Friction *= SharedSettings->BrakingFrictionFactor;
-	}
+			const bool bMaintainInputMagnitude = true;
+			Params.MoveInput = CharacterInputs->GetMoveInput_WorldSpace();
+			//Params.MoveInput = UPlanarConstraintUtils::ConstrainDirectionToPlane(MoverComp->GetPlanarConstraint(), CharacterInputs->GetMoveInput_WorldSpace(), bMaintainInputMagnitude);
+		}
+		else
+		{
+			Params.MoveInputType = EMoveInputType::None;
+			Params.MoveInput = FVector::ZeroVector;
+		}
 
-	OutProposedMove = UGroundMovementUtils::ComputeControlledGroundMove(Params);
+		Params.OrientationIntent = IntendedOrientation_WorldSpace;
+		Params.PriorVelocity = FVector::VectorPlaneProject(CurrentVelocity, MovementNormal);
+		Params.PriorOrientation = StartingSyncState->GetOrientation_WorldSpace();
+		Params.GroundNormal = MovementNormal;
+		Params.TurningRate = SharedSettings->TurningRate;
+		Params.TurningBoost = SharedSettings->TurningBoost;
+		Params.MaxSpeed = GetMaxSpeed();
+		Params.Acceleration = GetAcceleration();
+		Params.Deceleration = SharedSettings->Deceleration;
+		Params.DeltaSeconds = DeltaSeconds;
+		Params.WorldToGravityQuat = WorldToGravityTransform;
+		Params.UpDirection = UpDirection;
+		Params.bUseAccelerationForVelocityMove = SharedSettings->bUseAccelerationForVelocityMove;
+
+		if (Params.MoveInput.SizeSquared() > 0.f && !UMovementUtils::IsExceedingMaxSpeed(Params.PriorVelocity, SharedSettings->MaxSpeed))
+		{
+			Params.Friction = SharedSettings->GroundFriction;
+		}
+		else
+		{
+			Params.Friction = SharedSettings->bUseSeparateBrakingFriction ? SharedSettings->BrakingFriction : SharedSettings->GroundFriction;
+			Params.Friction *= SharedSettings->BrakingFrictionFactor;
+		}
+
+		OutProposedMove = UGroundMovementUtils::ComputeControlledGroundMove(Params);
+	}
 
 	if (SimBlackboard)
 	{
diff --git a/Engine/Plugins/Experimental/ChaosMover/Source/ChaosMover/Private/Utilities/ChaosGroundMovementUtils.cpp b/Engine/Plugins/Experimental/ChaosMover/Source/ChaosMover/Private/Utilities/ChaosGroundMovementUtils.cpp
--- a/Engine/Plugins/Experimental/ChaosMover/Source/ChaosMover/Private/Utilities/ChaosGroundMovementUtils.cpp
+++ b/Engine/Plugins/Experimental/ChaosMover/Source/ChaosMover/Private/Utilities/ChaosGroundMovementUtils.cpp
@@ -4,6 +4,7 @@
 
 #include "Chaos/PhysicsObjectInternalInterface.h"
 #include "MoveLibrary/FloorQueryUtils.h"
+#include "MoveLibrary/MovementUtils.h"
 
 #include UE_INLINE_GENERATED_CPP_BY_NAME(ChaosGroundMovementUtils)
 
@@ -30,3 +31,73 @@ Chaos::FPBDRigidParticleHandle* UChaosGroundMovementUtils::GetRigidParticleHandl
 
 	return nullptr;
 }
+
+FProposedMove UChaosGroundMovementUtils::ComputeRequestedMove(const FRequestedMoveParams& Params)
+{
+	FProposedMove OutProposedMove;
+
+	const FPlane MovementPlane(FVector::ZeroVector, Params.GroundNormal);
+	FVector NewVelocity = UMovementUtils::ConstrainToPlane(Params.PriorVelocity, MovementPlane, true);
+
+	// Copied from UCharacterMovementComponent::ApplyRequestedMove
+	const float RequestedSpeedSquared = Params.RequestedVelocity.SizeSquared();
+	float RequestedSpeed = 0.0f;
+	FVector RequestedMoveDir = FVector::ForwardVector;
+	const bool bZeroRequestedSpeed = RequestedSpeedSquared < UE_SMALL_NUMBER;
+	if (!bZeroRequestedSpeed)
+	{
+		// Compute requested speed from path following
+		RequestedSpeed = Params.RequestedVelocity.Size();
+		RequestedMoveDir = UMovementUtils::ConstrainToPlane(Params.RequestedVelocity / RequestedSpeed, MovementPlane, true);
+		RequestedSpeed = (Params.bRequestedMoveWithMaxSpeed ? Params.MaxSpeed : FMath::Min(Params.MaxSpeed, RequestedSpeed));
+
+		OutProposedMove.bHasDirIntent = true;
+		OutProposedMove.DirectionIntent = RequestedMoveDir;
+
+		// Compute actual requested velocity
+		const FVector MoveVelocity = RequestedMoveDir * RequestedSpeed;
+
+		// Compute new velocity
+		if (Params.bShouldComputeAcceleration && (Params.DeltaSeconds > UE_SMALL_NUMBER))
+		{
+			// How much do we need to accelerate to get to the new velocity?
+			FVector NewAcceleration = ((MoveVelocity - NewVelocity) / Params.DeltaSeconds);
+			NewAcceleration = NewAcceleration.GetClampedToMaxSize(Params.Acceleration);
+			NewVelocity = Params.PriorVelocity + NewAcceleration * Params.DeltaSeconds;
+		}
+		else
+		{
+			// Just set velocity directly.
+			// If decelerating we do so instantly, so we don't slide through the destination if we can't brake fast enough.
+			NewVelocity = MoveVelocity;
+		}
+	}
+
+	const bool bIsExceedingMaxSpeed = !Params.bRequestedMoveWithMaxSpeed && UMovementUtils::IsExceedingMaxSpeed(NewVelocity, Params.MaxSpeed);
+	if (bZeroRequestedSpeed || bIsExceedingMaxSpeed)
+	{
+		// Dampen velocity magnitude based on deceleration.
+		const FVector OldVelocity = NewVelocity;
+		const float VelSize = FMath::Max(NewVelocity.Size() - FMath::Abs(Params.Deceleration) * Params.DeltaSeconds, 0.f);
+		NewVelocity = NewVelocity.GetSafeNormal() * VelSize;
+
+		// Don't allow braking to lower us below max speed if we started above it.
+		if (bIsExceedingMaxSpeed && NewVelocity.SizeSquared() < FMath::Square(Params.MaxSpeed))
+		{
+			NewVelocity = OldVelocity.GetSafeNormal() * Params.MaxSpeed;
+		}
+	}
+
+	OutProposedMove.LinearVelocity = NewVelocity;
+
+	OutProposedMove.LinearVelocity = UMovementUtils::ConstrainToPlane(NewVelocity, MovementPlane, true);
+
+	// Linearly rotate in place
+	if (OutProposedMove.bHasDirIntent)
+	{
+		FRotator OrientationIntent = OutProposedMove.DirectionIntent.Rotation();
+		OutProposedMove.AngularVelocity = UMovementUtils::ComputeAngularVelocity(Params.PriorOrientation, OrientationIntent, Params.WorldToGravityQuat, Params.DeltaSeconds, Params.TurningRate);
+	}
+
+	return OutProposedMove;
+}
diff --git a/Engine/Plugins/Experimental/ChaosMover/Source/ChaosMover/Public/ChaosMover/Character/ChaosCharacterInputs.h b/Engine/Plugins/Experimental/ChaosMover/Source/ChaosMover/Public/ChaosMover/Character/ChaosCharacterInputs.h
--- a/Engine/Plugins/Experimental/ChaosMover/Source/ChaosMover/Public/ChaosMover/Character/ChaosCharacterInputs.h
+++ b/Engine/Plugins/Experimental/ChaosMover/Source/ChaosMover/Public/ChaosMover/Character/ChaosCharacterInputs.h
@@ -160,4 +160,46 @@ struct TStructOpsTypeTraits< FChaosMovementSettingsOverridesRemover > : public T
 		WithNetSerializer = true,
 		WithCopy = true
 	};
-};
\ No newline at end of file
+};
+
+/** Add this struct to the character inputs to apply a requested move the character */
+USTRUCT(MinimalAPI, BlueprintType)
+struct FChaosMoverRequestedMoveInputs : public FMoverDataStructBase
+{
+	GENERATED_BODY()
+
+public:
+	// Requested velocity in cm/s
+	UPROPERTY(BlueprintReadWrite, Category = "Chaos Mover")
+	FVector RequestedVelocity = FVector::ZeroVector;
+
+	UPROPERTY(BlueprintReadWrite, Category = "Chaos Mover")
+	bool bForceMaxSpeed = false;
+
+	UPROPERTY(BlueprintReadWrite, Category = "Chaos Mover")
+	bool bUseAcceleration = true;
+
+	virtual UScriptStruct* GetScriptStruct() const override
+	{
+		return StaticStruct();
+	}
+	virtual FMoverDataStructBase* Clone() const override
+	{
+		return new FChaosMoverRequestedMoveInputs(*this);
+	}
+	CHAOSMOVER_API virtual bool NetSerialize(FArchive& Ar, UPackageMap* Map, bool& bOutSuccess) override;
+	CHAOSMOVER_API virtual void ToString(FAnsiStringBuilderBase& Out) const override;
+	CHAOSMOVER_API virtual bool ShouldReconcile(const FMoverDataStructBase& AuthorityState) const override;
+	CHAOSMOVER_API virtual void Interpolate(const FMoverDataStructBase& From, const FMoverDataStructBase& To, float Pct) override;
+	CHAOSMOVER_API virtual void Merge(const FMoverDataStructBase& From) override;
+};
+
+template<>
+struct TStructOpsTypeTraits< FChaosMoverRequestedMoveInputs > : public TStructOpsTypeTraitsBase2< FChaosMoverRequestedMoveInputs >
+{
+	enum
+	{
+		WithNetSerializer = true,
+		WithCopy = true
+	};
+};
diff --git a/Engine/Plugins/Experimental/ChaosMover/Source/ChaosMover/Public/ChaosMover/Utilities/ChaosGroundMovementUtils.h b/Engine/Plugins/Experimental/ChaosMover/Source/ChaosMover/Public/ChaosMover/Utilities/ChaosGroundMovementUtils.h
--- a/Engine/Plugins/Experimental/ChaosMover/Source/ChaosMover/Public/ChaosMover/Utilities/ChaosGroundMovementUtils.h
+++ b/Engine/Plugins/Experimental/ChaosMover/Source/ChaosMover/Public/ChaosMover/Utilities/ChaosGroundMovementUtils.h
@@ -7,10 +7,51 @@
 
 #include "ChaosGroundMovementUtils.generated.h"
 
-#define UE_API CHAOSMOVER_API
-
 struct FFloorCheckResult;
 
+// Input parameters for controlled ground movement function
+USTRUCT(BlueprintType)
+struct FRequestedMoveParams
+{
+	GENERATED_BODY()
+
+	UPROPERTY(BlueprintReadWrite, EditAnywhere, Category = ChaosMover)
+	FQuat WorldToGravityQuat = FQuat::Identity;
+
+	UPROPERTY(BlueprintReadWrite, EditAnywhere, Category = ChaosMover)
+	FRotator PriorOrientation = FRotator::ZeroRotator;
+
+	UPROPERTY(BlueprintReadWrite, EditAnywhere, Category = ChaosMover)
+	FVector PriorVelocity = FVector::ZeroVector;
+
+	UPROPERTY(BlueprintReadWrite, EditAnywhere, Category = ChaosMover)
+	FVector RequestedVelocity = FVector::ZeroVector;
+
+	UPROPERTY(BlueprintReadWrite, EditAnywhere, Category = ChaosMover)
+	FVector GroundNormal = FVector::UpVector;
+
+	UPROPERTY(BlueprintReadWrite, EditAnywhere, Category = ChaosMover)
+	float TurningRate = 500.f;
+
+	UPROPERTY(BlueprintReadWrite, EditAnywhere, Category = ChaosMover)
+	float MaxSpeed = 800.f;
+
+	UPROPERTY(BlueprintReadWrite, EditAnywhere, Category = ChaosMover)
+	float Acceleration = 4000.f;
+
+	UPROPERTY(BlueprintReadWrite, EditAnywhere, Category = ChaosMover)
+	float Deceleration = 8000.f;
+
+	UPROPERTY(BlueprintReadWrite, EditAnywhere, Category = ChaosMover)
+	float DeltaSeconds = 0.f;
+
+	UPROPERTY(BlueprintReadWrite, EditAnywhere, Category = ChaosMover)
+	bool bRequestedMoveWithMaxSpeed = false;
+
+	UPROPERTY(BlueprintReadWrite, EditAnywhere, Category = ChaosMover)
+	bool bShouldComputeAcceleration = false;
+};
+
 UCLASS(MinimalAPI)
 class UChaosGroundMovementUtils : public UBlueprintFunctionLibrary
 {
@@ -19,9 +60,11 @@ class UChaosGroundMovementUtils : public UBlueprintFunctionLibrary
 public:
 	/** Computes the local velocity at the supplied position of the hit object in floor result */
 	UFUNCTION(BlueprintCallable, Category = ChaosMover)
-	static UE_API FVector ComputeLocalGroundVelocity_Internal(const FVector& Position, const FFloorCheckResult& FloorResult);
+	static CHAOSMOVER_API FVector ComputeLocalGroundVelocity_Internal(const FVector& Position, const FFloorCheckResult& FloorResult);
 
-	static UE_API Chaos::FPBDRigidParticleHandle* GetRigidParticleHandleFromFloorResult_Internal(const FFloorCheckResult& FloorResult);
-};
+	/** Computes the local velocity at the supplied position of the hit object in floor result */
+	UFUNCTION(BlueprintCallable, Category = ChaosMover)
+	static CHAOSMOVER_API FProposedMove ComputeRequestedMove(const FRequestedMoveParams& Params);
 
-#undef UE_API
+	static CHAOSMOVER_API Chaos::FPBDRigidParticleHandle* GetRigidParticleHandleFromFloorResult_Internal(const FFloorCheckResult& FloorResult);
+};
