Buildfix for ChaosGroundMovementUtils

diff --git a/Engine/Plugins/Experimental/ChaosMover/Source/ChaosMover/Private/Utilities/ChaosGroundMovementUtils.cpp b/Engine/Plugins/Experimental/ChaosMover/Source/ChaosMover/Private/Utilities/ChaosGroundMovementUtils.cpp
--- a/Engine/Plugins/Experimental/ChaosMover/Source/ChaosMover/Private/Utilities/ChaosGroundMovementUtils.cpp
+++ b/Engine/Plugins/Experimental/ChaosMover/Source/ChaosMover/Private/Utilities/ChaosGroundMovementUtils.cpp
@@ -96,7 +96,7 @@ FProposedMove UChaosGroundMovementUtils::ComputeRequestedMove(const FRequestedMo
 	if (OutProposedMove.bHasDirIntent)
 	{
 		FRotator OrientationIntent = OutProposedMove.DirectionIntent.Rotation();
-		OutProposedMove.AngularVelocity = UMovementUtils::ComputeAngularVelocity(Params.PriorOrientation, OrientationIntent, Params.WorldToGravityQuat, Params.DeltaSeconds, Params.TurningRate);
+		OutProposedMove.AngularVelocityDegrees = UMovementUtils::ComputeAngularVelocityDegrees(Params.PriorOrientation, OrientationIntent, Params.DeltaSeconds, Params.TurningRate);
 	}
 
 	return OutProposedMove;
