Local only Instant Movement Effect, for debugging desyncs. Console Command to teleport the player locally only and force a resim.

diff --git a/Engine/Plugins/Experimental/ChaosMover/Source/ChaosMover/Private/ChaosMoverConsoleVariables.cpp b/Engine/Plugins/Experimental/ChaosMover/Source/ChaosMover/Private/ChaosMoverConsoleVariables.cpp
--- a/Engine/Plugins/Experimental/ChaosMover/Source/ChaosMover/Private/ChaosMoverConsoleVariables.cpp
+++ b/Engine/Plugins/Experimental/ChaosMover/Source/ChaosMover/Private/ChaosMoverConsoleVariables.cpp
@@ -4,6 +4,13 @@
 
 #include "HAL/IConsoleManager.h"
 
+#include "ChaosMover/ChaosMoverSimulationTypes.h"
+#include "ChaosMover/ChaosMoverSimulation.h"
+#include "MoverComponent.h"
+#include "GameFramework/PlayerController.h"
+#include "GameFramework/Pawn.h"
+#include "Engine/World.h"
+
 namespace UE::ChaosMover
 {
 	namespace CVars
@@ -30,6 +37,42 @@ namespace UE::ChaosMover
 
 		int32 InstantMovementEffectIDHistorySize = 30;
 		FAutoConsoleVariableRef CVarChaosMoverInstantMovementEffectIDHistorySize(TEXT("ChaosMover.Networking.InstantMovementEffectIDHistorySize"),
-			bSkipGenerateMoveIfOverridden, TEXT("Number of past frames for which we keep the last seen instant movement effect IDs, to avoid duplicate processing of events when input is repeated\n"));
+			InstantMovementEffectIDHistorySize, TEXT("Number of past frames for which we keep the last seen instant movement effect IDs, to avoid duplicate processing of events when input is repeated\n"));
+
+		static FAutoConsoleCommandWithWorld TeleportOnlyClientCommand
+		(
+			TEXT("ChaosMover.Debug.TeleportClientOnly"),
+			TEXT("Sends an instant movement effect to teleport the locally controlled character if it has a UMoverComponent"),
+			FConsoleCommandWithWorldDelegate::CreateStatic([](UWorld* World)
+				{
+					APlayerController* LocalPlayerController = World ? World->GetFirstPlayerController() : nullptr;
+					APawn* LocalPawn = LocalPlayerController ? LocalPlayerController->GetPawn() : nullptr;
+					UMoverComponent* MoverComponent = LocalPawn ? LocalPawn->FindComponentByClass<UMoverComponent>() : nullptr;
+					if (!MoverComponent)
+					{
+						return;
+					}
+
+					auto TeleportLambda = [](FApplyMovementEffectParams_Async& ApplyEffectParams, FMoverSyncState& OutputState)->bool
+					{
+						if (UChaosMoverSimulation* ChaosMoverSimulation = Cast<UChaosMoverSimulation>(ApplyEffectParams.Simulation))
+						{
+							if (const FMoverDefaultSyncState* DefaultSyncState = OutputState.SyncStateCollection.FindDataByType<FMoverDefaultSyncState>())
+							{
+								static FVector TeleportDeltaLocation(0, 0, 500);
+								FVector TargetLocation = DefaultSyncState->GetLocation_WorldSpace() + TeleportDeltaLocation;
+								ChaosMoverSimulation->AttemptTeleport(*ApplyEffectParams.TimeStep, FTransform(DefaultSyncState->GetOrientation_WorldSpace(), TargetLocation), /* bUseActorRotation = */ true, OutputState);
+								return true;
+							}
+						}
+						return false;
+					};
+
+					TSharedPtr<FAsyncLocalOnlyInstantMovementEffect> EffectPtr = MakeShared<FAsyncLocalOnlyInstantMovementEffect>();
+					static FName EffectName("TeleportOnlyInstantMovementEffect");
+					EffectPtr->OptionalName = EffectName;
+					EffectPtr->AsyncFunction = TeleportLambda;
+					MoverComponent->QueueInstantMovementEffect(EffectPtr);
+				}));
 	}
 }
\ No newline at end of file
diff --git a/Engine/Plugins/Experimental/ChaosMover/Source/ChaosMover/Private/ChaosMoverSimulationTypes.cpp b/Engine/Plugins/Experimental/ChaosMover/Source/ChaosMover/Private/ChaosMoverSimulationTypes.cpp
--- a/Engine/Plugins/Experimental/ChaosMover/Source/ChaosMover/Private/ChaosMoverSimulationTypes.cpp
+++ b/Engine/Plugins/Experimental/ChaosMover/Source/ChaosMover/Private/ChaosMoverSimulationTypes.cpp
@@ -262,3 +262,36 @@ void FChaosNetInstantMovementEffectsQueue::Decay(float DecayAmount)
 {
 
 }
+
+bool FAsyncLocalOnlyInstantMovementEffect::ApplyMovementEffect_Async(FApplyMovementEffectParams_Async& ApplyEffectParams, FMoverSyncState& OutputState)
+{
+	if (bExecuteLocally && AsyncFunction.IsSet())
+	{
+		return AsyncFunction(ApplyEffectParams, OutputState);
+	}
+	return false;
+}
+
+FInstantMovementEffect* FAsyncLocalOnlyInstantMovementEffect::Clone() const
+{
+	return new FAsyncLocalOnlyInstantMovementEffect(*this);
+}
+
+void FAsyncLocalOnlyInstantMovementEffect::NetSerialize(FArchive& Ar)
+{
+	// Serialize nothing out, and when loading (receiving from the network), set to not execute locally
+	if (Ar.IsLoading())
+	{
+		bExecuteLocally = false;
+	}
+}
+
+UScriptStruct* FAsyncLocalOnlyInstantMovementEffect::GetScriptStruct() const
+{
+	return FAsyncLocalOnlyInstantMovementEffect::StaticStruct();
+}
+
+FString FAsyncLocalOnlyInstantMovementEffect::ToSimpleString() const
+{
+	return FString::Printf(TEXT("FAsyncLocalOnlyInstantMovementEffect %s"), *OptionalName.ToString());
+}
diff --git a/Engine/Plugins/Experimental/ChaosMover/Source/ChaosMover/Public/ChaosMover/ChaosMoverSimulationTypes.h b/Engine/Plugins/Experimental/ChaosMover/Source/ChaosMover/Public/ChaosMover/ChaosMoverSimulationTypes.h
--- a/Engine/Plugins/Experimental/ChaosMover/Source/ChaosMover/Public/ChaosMover/ChaosMoverSimulationTypes.h
+++ b/Engine/Plugins/Experimental/ChaosMover/Source/ChaosMover/Public/ChaosMover/ChaosMoverSimulationTypes.h
@@ -394,4 +394,26 @@ struct FChaosNetInstantMovementEffectsQueue : public FMoverDataStructBase
 	virtual void Decay(float DecayAmount) override;
 };
 
+// This instant movement effect is ChaosMover only and applies only on the instance issuing it
+// It will execute AsyncFunction on the physics thread if set
+USTRUCT()
+struct FAsyncLocalOnlyInstantMovementEffect : public FInstantMovementEffect
+{
+	GENERATED_BODY()
+
+	UE_API FAsyncLocalOnlyInstantMovementEffect() {}
+	virtual ~FAsyncLocalOnlyInstantMovementEffect() {}
+
+	FName OptionalName;
+	// Async Function executed in ApplyMovementEffect_Async (only locally)
+	TFunction<bool(FApplyMovementEffectParams_Async&, FMoverSyncState&)> AsyncFunction;
+	bool bExecuteLocally = true;
+
+	UE_API virtual bool ApplyMovementEffect_Async(FApplyMovementEffectParams_Async& ApplyEffectParams, FMoverSyncState& OutputState) override;
+	UE_API virtual FInstantMovementEffect* Clone() const override;
+	UE_API virtual void NetSerialize(FArchive& Ar) override;
+	UE_API virtual UScriptStruct* GetScriptStruct() const override;
+	UE_API virtual FString ToSimpleString() const override;
+};
+
 #undef UE_API
