ChaosMover: Clear queued instant movement effects after they are injected in to the simulation.
#rnx

diff --git a/Engine/Plugins/Experimental/ChaosMover/Source/ChaosMover/Private/Backends/ChaosMoverBackend.cpp b/Engine/Plugins/Experimental/ChaosMover/Source/ChaosMover/Private/Backends/ChaosMoverBackend.cpp
--- a/Engine/Plugins/Experimental/ChaosMover/Source/ChaosMover/Private/Backends/ChaosMoverBackend.cpp
+++ b/Engine/Plugins/Experimental/ChaosMover/Source/ChaosMover/Private/Backends/ChaosMoverBackend.cpp
@@ -588,6 +588,7 @@ void UChaosMoverBackendComponent::InjectInstantMovementEffectsIntoSim(const FMov
 			UE_LOG(LogChaosMover, Verbose, TEXT("(%s) ProduceInputData: Async simulation received an Instant Movement Effect from the Mover Component, scheduled for frame %d at frame %d: %s"),
 				*ToString(NetMode), ScheduledEffect.ExecutionServerFrame, TimeStep.ServerFrame, ScheduledEffect.Effect.IsValid() ? *ScheduledEffect.Effect->ToSimpleString() : TEXT("INVALID INSTANT EFFECT"));
 		}
+		MoverComp.ClearQueuedInstantMovementEffects();
 	}
 }
 
