[Mover] Allow custom chaos landed event data

Rename `OnLanded` delegate to be `OnLandedDelegate`.
Added virtual `OnLanded` function for child classes to use instead of having to subscribe to their own delegate

#rnx

diff --git a/Engine/Plugins/Experimental/ChaosMover/Config/BaseChaosMover.ini b/Engine/Plugins/Experimental/ChaosMover/Config/BaseChaosMover.ini
new file mode 100644
--- /dev/null
+++ b/Engine/Plugins/Experimental/ChaosMover/Config/BaseChaosMover.ini
@@ -0,0 +1,2 @@
+﻿[CoreRedirects]
++PropertyRedirects=(OldName="/Script/ChaosMover.ChaosCharacterMoverComponent.OnLanded",NewName="/Script/ChaosMover.ChaosCharacterMoverComponent.OnLandedDelegate")
\ No newline at end of file
diff --git a/Engine/Plugins/Experimental/ChaosMover/Source/ChaosMover/Private/Character/ChaosCharacterMoverComponent.cpp b/Engine/Plugins/Experimental/ChaosMover/Source/ChaosMover/Private/Character/ChaosCharacterMoverComponent.cpp
--- a/Engine/Plugins/Experimental/ChaosMover/Source/ChaosMover/Private/Character/ChaosCharacterMoverComponent.cpp
+++ b/Engine/Plugins/Experimental/ChaosMover/Source/ChaosMover/Private/Character/ChaosCharacterMoverComponent.cpp
@@ -40,7 +40,7 @@ void UChaosCharacterMoverComponent::ProcessSimulationEvent(const FMoverSimulatio
 
 	if (const FLandedEventData* LandedData = EventData.CastTo<FLandedEventData>())
 	{
-		OnLanded.Broadcast(LandedData->NewModeName, LandedData->HitResult);
+		NotifyOnLanded(*LandedData);
 	}
 	if (const FStanceModifiedEventData* StanceModifiedEvent = EventData.CastTo<FStanceModifiedEventData>())
 	{
@@ -165,6 +165,17 @@ void UChaosCharacterMoverComponent::OnStanceModified(const FStanceModifiedEventD
 	OnStanceChanged.Broadcast(EventData.OldStance, EventData.NewStance);
 }
 
+void UChaosCharacterMoverComponent::NotifyOnLanded(const FLandedEventData& LandedEvent)
+{
+	OnLanded(LandedEvent);
+	OnLandedDelegate.Broadcast(this, LandedEvent);
+}
+
+void UChaosCharacterMoverComponent::OnLanded(const FLandedEventData& LandedEvent)
+{
+	
+}
+
 void UChaosCharacterMoverComponent::OverrideMovementSettings(const FChaosMovementSettingsOverrides Overrides)
 {
 	ModeToOverrideSettings = Overrides.ModeName;
diff --git a/Engine/Plugins/Experimental/ChaosMover/Source/ChaosMover/Private/Character/Transitions/ChaosCharacterLandingCheck.cpp b/Engine/Plugins/Experimental/ChaosMover/Source/ChaosMover/Private/Character/Transitions/ChaosCharacterLandingCheck.cpp
--- a/Engine/Plugins/Experimental/ChaosMover/Source/ChaosMover/Private/Character/Transitions/ChaosCharacterLandingCheck.cpp
+++ b/Engine/Plugins/Experimental/ChaosMover/Source/ChaosMover/Private/Character/Transitions/ChaosCharacterLandingCheck.cpp
@@ -94,8 +94,25 @@ void UChaosCharacterLandingCheck::Trigger_Implementation(const FSimulationTickPa
 			FFloorCheckResult FloorResult;
 			if (Blackboard->TryGet(CommonBlackboard::LastFloorResult, FloorResult))
 			{
-				Sim->AddEvent(MakeShared<FLandedEventData>(Params.TimeStep.BaseSimTimeMs, FloorResult.HitResult, TransitionToGroundMode));
+				const TSharedPtr<FLandedEventData> LandedEventData = AllocLandedEventData();
+				if (LandedEventData.IsValid())
+				{
+					PopulateLandedEventData(Params, FloorResult, *LandedEventData);
+					Sim->AddEvent(LandedEventData);
+				}
 			}
 		}
 	}
-}
\ No newline at end of file
+}
+
+TSharedPtr<FLandedEventData> UChaosCharacterLandingCheck::AllocLandedEventData()
+{
+	return MakeShared<FLandedEventData>();
+}
+
+void UChaosCharacterLandingCheck::PopulateLandedEventData(const FSimulationTickParams& Params, const FFloorCheckResult& FloorResult, FLandedEventData& OutLandedEvent)
+{
+	OutLandedEvent.EventTimeMs = Params.TimeStep.BaseSimTimeMs;
+	OutLandedEvent.HitResult = FloorResult.HitResult;
+	OutLandedEvent.NewModeName = TransitionToGroundMode;
+}
diff --git a/Engine/Plugins/Experimental/ChaosMover/Source/ChaosMover/Public/ChaosMover/Character/ChaosCharacterMoverComponent.h b/Engine/Plugins/Experimental/ChaosMover/Source/ChaosMover/Public/ChaosMover/Character/ChaosCharacterMoverComponent.h
--- a/Engine/Plugins/Experimental/ChaosMover/Source/ChaosMover/Public/ChaosMover/Character/ChaosCharacterMoverComponent.h
+++ b/Engine/Plugins/Experimental/ChaosMover/Source/ChaosMover/Public/ChaosMover/Character/ChaosCharacterMoverComponent.h
@@ -4,12 +4,13 @@
 
 #include "ChaosMover/ChaosMoverSimulationTypes.h"
 #include "DefaultMovementSet/CharacterMoverComponent.h"
+#include "DefaultMovementSet/CharacterMoverSimulationTypes.h"
 #include "MoveLibrary/FloorQueryUtils.h"
 
 #include "ChaosCharacterMoverComponent.generated.h"
 
-// Fired after the actor lands on a valid surface. First param is the name of the mode this actor is in after landing. Second param is the hit result from hitting the floor.
-DECLARE_DYNAMIC_MULTICAST_DELEGATE_TwoParams(FChaosMover_OnLanded, const FName&, NextMovementModeName, const FHitResult&, HitResult);
+// Fired after the actor lands on a valid surface.
+DECLARE_DYNAMIC_MULTICAST_DELEGATE_TwoParams(FChaosMover_OnLanded, class UChaosCharacterMoverComponent*, ChaosMoverComponent, FLandedEventData, LandedEventData);
 
 
 UCLASS(MinimalAPI, BlueprintType, Blueprintable, meta = (BlueprintSpawnableComponent))
@@ -28,8 +29,8 @@ public:
 	CHAOSMOVER_API virtual void UnCrouch() override;
 
 	// Broadcast when this actor lands on a valid surface.
-	UPROPERTY(BlueprintAssignable, Category = "Chaos Mover")
-	FChaosMover_OnLanded OnLanded;
+	UPROPERTY(BlueprintAssignable, Category = "Chaos Mover", DisplayName = "On Landed")
+	FChaosMover_OnLanded OnLandedDelegate;
 
 	// Launch the character using either impulse or velocity
 	// Note: This will only trigger a launch if a launch transition is implemented on the current movement mode
@@ -54,6 +55,9 @@ protected:
 	CHAOSMOVER_API void ClearQueuedMode();
 	CHAOSMOVER_API virtual void OnStanceModified(const FStanceModifiedEventData& EventData);
 
+	CHAOSMOVER_API void NotifyOnLanded(const FLandedEventData& LandedEvent);
+	CHAOSMOVER_API virtual void OnLanded(const FLandedEventData& LandedEvent);
+
 	FName ModeToOverrideSettings = NAME_None;
 	float MaxSpeedOverride = 0.0f;
 	float AccelerationOverride = 0.0f;
diff --git a/Engine/Plugins/Experimental/ChaosMover/Source/ChaosMover/Public/ChaosMover/Character/Transitions/ChaosCharacterLandingCheck.h b/Engine/Plugins/Experimental/ChaosMover/Source/ChaosMover/Public/ChaosMover/Character/Transitions/ChaosCharacterLandingCheck.h
--- a/Engine/Plugins/Experimental/ChaosMover/Source/ChaosMover/Public/ChaosMover/Character/Transitions/ChaosCharacterLandingCheck.h
+++ b/Engine/Plugins/Experimental/ChaosMover/Source/ChaosMover/Public/ChaosMover/Character/Transitions/ChaosCharacterLandingCheck.h
@@ -3,6 +3,7 @@
 #pragma once
 
 #include "ChaosMover/ChaosMovementModeTransition.h"
+#include "DefaultMovementSet/CharacterMoverSimulationTypes.h"
 
 #include "ChaosCharacterLandingCheck.generated.h"
 
@@ -27,6 +28,13 @@ public:
 	/** Name of movement mode to transition to when landing on ground */
 	UPROPERTY(BlueprintReadWrite, EditAnywhere, Category = Parameters)
 	FName TransitionToGroundMode;
+
+protected:
+	/** Create the event data we will use for the landed event*/
+	UE_API virtual TSharedPtr<FLandedEventData> AllocLandedEventData();
+	
+	/** Used to fill the landed event data before adding the Landed Event into the simulation */
+	UE_API virtual void PopulateLandedEventData(const FSimulationTickParams& Params, const FFloorCheckResult& FloorResult, FLandedEventData& OutLandedEvent);
 };
 
 #undef UE_API
