Bubble Resim: ChaosMover ShouldResim function to determine if we need to update during resim. This is only used when bubble resim is on.

diff --git a/Engine/Plugins/Experimental/ChaosMover/Source/ChaosMover/Private/Backends/ChaosMoverAsyncCallback.cpp b/Engine/Plugins/Experimental/ChaosMover/Source/ChaosMover/Private/Backends/ChaosMoverAsyncCallback.cpp
--- a/Engine/Plugins/Experimental/ChaosMover/Source/ChaosMover/Private/Backends/ChaosMoverAsyncCallback.cpp
+++ b/Engine/Plugins/Experimental/ChaosMover/Source/ChaosMover/Private/Backends/ChaosMoverAsyncCallback.cpp
@@ -3,6 +3,7 @@
 #include "ChaosMoverAsyncCallback.h"
 
 #include "Chaos/Framework/Parallel.h"
+#include "Chaos/ResimConsoleVariables.h"
 #include "ChaosMover/Backends/ChaosMoverBackend.h"
 #include "ChaosMover/ChaosMoverConsoleVariables.h"
 #include "ChaosMover/ChaosMoverSimulation.h"
@@ -13,6 +14,8 @@ namespace UE::ChaosMover
 {
 	void FAsyncCallback::ProcessInputs_Internal(int32 PhysicsStep)
 	{
+		QUICK_SCOPE_CYCLE_COUNTER(ChaosMover_AsyncCallback_ProcessInputs_Internal);
+
 		const FAsyncCallbackInput* AsyncInput = GetConsumerInput_Internal();
 
 		if (!AsyncInput || AsyncInput->InputData.IsEmpty())
@@ -23,19 +26,28 @@ namespace UE::ChaosMover
 		FMoverTimeStep TimeStep;
 		GetCurrentMoverTimeStep(AsyncInput, TimeStep);
 
-		auto LambdaParallelUpdate = [&TimeStep, &AsyncInput, PhysicsStep](int32 Idx) {
-			if (TStrongObjectPtr<UChaosMoverBackendComponent> Backend = AsyncInput->Backends[Idx].Pin())
-			{
-				Backend->GetSimulation()->ProcessInputs(PhysicsStep, TimeStep, AsyncInput->InputData[Idx]);
+		// This array of resimulated backend indices is only used in the resim case
+		const TArray<int32, TInlineAllocator<64>>* ResimulatedBackendIndicesPtr = TimeStep.bIsResimulating ? DetermineResimulatedBackendIndices(AsyncInput->Backends) : nullptr;
+
+		auto LambdaParallelUpdate = [&TimeStep, &AsyncInput, PhysicsStep, ResimulatedBackendIndicesPtr](int32 Idx) {
+			int32 BackendIndex = ResimulatedBackendIndicesPtr ? (*ResimulatedBackendIndicesPtr)[Idx] : Idx;
+			if (TStrongObjectPtr<UChaosMoverBackendComponent> Backend = AsyncInput->Backends[BackendIndex].Pin())
+			{				
+				Backend->GetSimulation()->ProcessInputs(PhysicsStep, TimeStep, AsyncInput->InputData[BackendIndex]);
 			}
 			};
 
+		// Backends are added to AsyncInput in InjectInputs_EXTERNAL
+		const int32 NumBackends = AsyncInput->Backends.Num();
+		const int32 NumSimulatedBackends = ResimulatedBackendIndicesPtr ? ResimulatedBackendIndicesPtr->Num() : NumBackends;
 		const bool ForceSingleThread = UE::ChaosMover::CVars::bForceSingleThreadedPT;
-		Chaos::PhysicsParallelFor(AsyncInput->Backends.Num(), LambdaParallelUpdate, ForceSingleThread);
+		Chaos::PhysicsParallelFor(NumSimulatedBackends, LambdaParallelUpdate, ForceSingleThread);
 	}
 
 	void FAsyncCallback::OnPreSimulate_Internal()
 	{
+		QUICK_SCOPE_CYCLE_COUNTER(ChaosMover_AsyncCallback_OnPreSimulate_Internal);
+
 		const FAsyncCallbackInput* AsyncInput = GetConsumerInput_Internal();
 
 		if (!AsyncInput || AsyncInput->InputData.IsEmpty())
@@ -47,24 +59,36 @@ namespace UE::ChaosMover
 		GetCurrentMoverTimeStep(AsyncInput, TimeStep);
 
 		FAsyncCallbackOutput& AsyncOutput = GetProducerOutputData_Internal();
-		AsyncOutput.OutputData.SetNum(AsyncInput->InputData.Num());
-		AsyncOutput.TimeStep.SetNum(AsyncInput->InputData.Num());
-
-		auto LambdaParallelUpdate = [&TimeStep, &AsyncInput, &AsyncOutput, this](int32 Idx) {
-			AsyncOutput.TimeStep[Idx] = TimeStep;
-			if (TStrongObjectPtr<UChaosMoverBackendComponent> Backend = AsyncInput->Backends[Idx].Pin())
+		
+		const int32 NumInputs = AsyncInput->InputData.Num();
+		AsyncOutput.OutputData.SetNum(NumInputs);
+		AsyncOutput.TimeStep.SetNumUninitialized(NumInputs);		
+
+		const int32 NumBackends = AsyncInput->Backends.Num();
+		ensure(NumInputs == NumBackends);
+		AsyncOutput.Backends = AsyncInput->Backends;
+
+		// This array of resimulated backend indices is only used in the resim case
+		const TArray<int32, TInlineAllocator<64>>* ResimulatedBackendIndicesPtr = (TimeStep.bIsResimulating && Chaos::ResimConsoleVars::bIsBubbleResimEnabled) ? &ResimulatedBackendIndices : nullptr;
+		
+		auto LambdaParallelUpdate = [&TimeStep, &AsyncInput, &AsyncOutput, ResimulatedBackendIndicesPtr, this](int32 Idx) {
+			int32 BackendIndex = ResimulatedBackendIndicesPtr ? (*ResimulatedBackendIndicesPtr)[Idx] : Idx;
+			AsyncOutput.TimeStep[BackendIndex] = TimeStep;
+			if (TStrongObjectPtr<UChaosMoverBackendComponent> Backend = AsyncInput->Backends[BackendIndex].Pin())
 			{
-				Backend->GetSimulation()->SimulationTick(TimeStep, AsyncInput->InputData[Idx], AsyncOutput.OutputData[Idx]);
-				AsyncOutput.Backends.Add(Backend.Get());
+				Backend->GetSimulation()->SimulationTick(TimeStep, AsyncInput->InputData[BackendIndex], AsyncOutput.OutputData[BackendIndex]);
 			}
 		};
 
 		const bool ForceSingleThread = UE::ChaosMover::CVars::bForceSingleThreadedPT;
-		Chaos::PhysicsParallelFor(AsyncInput->Backends.Num(), LambdaParallelUpdate, ForceSingleThread);
+		const int32 NumSimulatedBackends = ResimulatedBackendIndicesPtr ? ResimulatedBackendIndicesPtr->Num() : NumBackends;
+		Chaos::PhysicsParallelFor(NumSimulatedBackends, LambdaParallelUpdate, ForceSingleThread);
 	}
 
 	void FAsyncCallback::OnContactModification_Internal(Chaos::FCollisionContactModifier& Modifier)
 	{
+		QUICK_SCOPE_CYCLE_COUNTER(ChaosMover_OnContactModification_Internal);
+
 		const FAsyncCallbackInput* AsyncInput = GetConsumerInput_Internal();
 
 		if (!AsyncInput)
@@ -75,12 +99,20 @@ namespace UE::ChaosMover
 		FMoverTimeStep TimeStep;
 		GetCurrentMoverTimeStep(AsyncInput, TimeStep);
 
+		// This array of resimulated backend indices is only used in the resim case
+		const TArray<int32, TInlineAllocator<64>>* ResimulatedBackendIndicesPtr =
+			(TimeStep.bIsResimulating && Chaos::ResimConsoleVars::bIsBubbleResimEnabled && !Chaos::ResimConsoleVars::bCanFrozenParticlesModifyContacts) ?
+			&ResimulatedBackendIndices : nullptr;
+
+		const int32 NumSimulatedBackends = ResimulatedBackendIndicesPtr ? ResimulatedBackendIndicesPtr->Num() : AsyncInput->Backends.Num();
+
 		const FAsyncCallbackOutput& AsyncOutput = GetProducerOutputData_Internal();
-		for (int Idx = 0; Idx < AsyncInput->Backends.Num(); ++Idx)
+		for (int Idx = 0; Idx < NumSimulatedBackends; ++Idx)
 		{
-			if (TStrongObjectPtr<UChaosMoverBackendComponent> Backend = AsyncInput->Backends[Idx].Pin())
+			int32 BackendIndex = ResimulatedBackendIndicesPtr ? (*ResimulatedBackendIndicesPtr)[Idx] : Idx;
+			if (TStrongObjectPtr<UChaosMoverBackendComponent> Backend = AsyncInput->Backends[BackendIndex].Pin())
 			{
-				Backend->GetSimulation()->ModifyContacts(TimeStep, AsyncInput->InputData[Idx], AsyncOutput.OutputData[Idx], Modifier);
+				Backend->GetSimulation()->ModifyContacts(TimeStep, AsyncInput->InputData[BackendIndex], AsyncOutput.OutputData[BackendIndex], Modifier);
 			}
 		}
 	}
@@ -95,4 +127,31 @@ namespace UE::ChaosMover
 			OutMoverTimeStep.bIsResimulating = AsyncInput->PhysicsSolver->GetEvolution()->IsResimming();
 		}
 	}
+
+	const TArray<int32, TInlineAllocator<64>>* FAsyncCallback::DetermineResimulatedBackendIndices(const TArray<TWeakObjectPtr<UChaosMoverBackendComponent>>& InBackends)
+	{
+		if (Chaos::ResimConsoleVars::bIsBubbleResimEnabled)
+		{
+			QUICK_SCOPE_CYCLE_COUNTER(ChaosMover_DetermineResimulatedBackendIndices);
+
+			int32 NumBackends = InBackends.Num();
+			ResimulatedBackendIndices.Reserve(NumBackends);
+			ResimulatedBackendIndices.Empty();
+
+			for (int32 BackendIndex = 0; BackendIndex < NumBackends; ++BackendIndex)
+			{
+				if (TStrongObjectPtr<UChaosMoverBackendComponent> Backend = InBackends[BackendIndex].Pin())
+				{
+					if (Backend->ShouldResim())
+					{
+						ResimulatedBackendIndices.Add(BackendIndex);
+					}
+				}
+			}
+
+			return &ResimulatedBackendIndices;
+		}
+
+		return nullptr;
+	}
 }
\ No newline at end of file
diff --git a/Engine/Plugins/Experimental/ChaosMover/Source/ChaosMover/Private/Backends/ChaosMoverAsyncCallback.h b/Engine/Plugins/Experimental/ChaosMover/Source/ChaosMover/Private/Backends/ChaosMoverAsyncCallback.h
--- a/Engine/Plugins/Experimental/ChaosMover/Source/ChaosMover/Private/Backends/ChaosMoverAsyncCallback.h
+++ b/Engine/Plugins/Experimental/ChaosMover/Source/ChaosMover/Private/Backends/ChaosMoverAsyncCallback.h
@@ -51,5 +51,13 @@ namespace UE::ChaosMover
 		virtual void OnContactModification_Internal(Chaos::FCollisionContactModifier& Modifier) override;
 
 		void GetCurrentMoverTimeStep(const FAsyncCallbackInput* AsyncInput, FMoverTimeStep& OutMoverTimeStep) const;
+
+		// With bubble resim enabled, determines which backends should resimulate and returns a pointer to ResimulatedBackendIndices
+		// Only call this function during resim frames
+		const TArray<int32, TInlineAllocator<64>>* DetermineResimulatedBackendIndices(const TArray<TWeakObjectPtr<UChaosMoverBackendComponent>>& InBackends);
+
+		// Array of resimulated backend indices. This array is used only in the resim case.
+		// During a resim frame, it is reset and filled in ProcessInputs_Internal and reused in subsequent callbacks on that same frame, such as OnPreSimulate_Internal
+		TArray<int32, TInlineAllocator<64>> ResimulatedBackendIndices;
 	};
 }
\ No newline at end of file
diff --git a/Engine/Plugins/Experimental/ChaosMover/Source/ChaosMover/Private/Backends/ChaosMoverBackend.cpp b/Engine/Plugins/Experimental/ChaosMover/Source/ChaosMover/Private/Backends/ChaosMoverBackend.cpp
--- a/Engine/Plugins/Experimental/ChaosMover/Source/ChaosMover/Private/Backends/ChaosMoverBackend.cpp
+++ b/Engine/Plugins/Experimental/ChaosMover/Source/ChaosMover/Private/Backends/ChaosMoverBackend.cpp
@@ -275,6 +275,8 @@ Chaos::FPhysicsObject* UChaosMoverBackendComponent::GetPhysicsObject() const
 
 Chaos::FPBDRigidParticle* UChaosMoverBackendComponent::GetControlledParticle() const
 {
+	Chaos::EnsureIsInGameThreadContext();
+
 	if (Chaos::FPhysicsObject* PhysicsObject = GetPhysicsObject())
 	{
 		return FPhysicsObjectExternalInterface::LockRead(PhysicsObject)->GetRigidParticle(PhysicsObject);
@@ -450,6 +452,11 @@ bool UChaosMoverBackendComponent::IsAsync() const
 	return true;
 }
 
+bool UChaosMoverBackendComponent::ShouldResim() const
+{
+	return Simulation ? Simulation->ShouldResim() : IMoverBackendLiaisonInterface::ShouldResim();
+}
+
 UChaosMoverSimulation* UChaosMoverBackendComponent::GetSimulation()
 {
 	return Simulation;
@@ -484,21 +491,32 @@ void UChaosMoverBackendComponent::ProduceInputData(int32 PhysicsStep, int32 NumS
 	if (bIsGeneratingInputsLocally)
 	{
 		GenerateInput(TimeStep, InputData);
-		// Input generating simulations inject instant movement effects from the Game Thread into the simulation input collection
-		// It will be used directly by the state machine (same as when it is received from the network or overwritten during resim),
-		// and will also be networked via NetInputCmd
-		InjectInstantMovementEffectsIntoInput(TimeStep, InputData);
 	}
-	else
+
+	UMoverComponent& MoverComp = GetMoverComponent();
+	// Add all queued instant effects to the input so it can be networked
+	const TArray<FScheduledInstantMovementEffect>& ScheduledInstantMovementEffects = MoverComp.GetQueuedInstantMovementEffects();
+	if (!ScheduledInstantMovementEffects.IsEmpty())
 	{
-		// Transfer queued instant movement effects to the simulation
-		// Instant effects should still be consumed by a non input generating actor on a server if it is not controlled remotely (this is an approximation of "no other instance is input producing")
-		InjectInstantMovementEffectsIntoSim(TimeStep);
+		if (bIsGeneratingInputsLocally)
+		{
+			// Input generating simulations inject instant movement effects from the Game Thread into the simulation input collection
+			// It will be used directly by the state machine (same as when it is received from the network or overwritten during resim),
+			// and will also be networked via NetInputCmd
+			InjectInstantMovementEffectsIntoInput(TimeStep, InputData);
+		}
+		else
+		{
+			// Transfer queued instant movement effects to the simulation
+			// Instant effects should still be consumed by a non input generating actor on a server if it is not controlled remotely (this is an approximation of "no other instance is input producing")
+			InjectInstantMovementEffectsIntoSim(TimeStep);
+		}
+		// We clear queued instant movement effects systematically to avoid accumulation, even if they ended up being discarded
+		MoverComp.ClearQueuedInstantMovementEffects();
 	}
 
 	// Add default simulation input data
 	FChaosMoverSimulationDefaultInputs& SimInputs = Simulation->GetLocalSimInput_Mutable().FindOrAddMutableDataByType<FChaosMoverSimulationDefaultInputs>();
-	UMoverComponent& MoverComp = GetMoverComponent();
 
 	SimInputs.Gravity = MoverComp.GetGravityAcceleration();
 	SimInputs.UpDir = MoverComp.GetUpDirection();
@@ -548,30 +566,27 @@ void UChaosMoverBackendComponent::InjectInstantMovementEffectsIntoInput(const FM
 {
 	UMoverComponent& MoverComp = GetMoverComponent();
 	// Add all queued instant effects to the input so it can be networked
-	const TArray<FScheduledInstantMovementEffect>& ScheduledInstantMovementEffect = MoverComp.GetQueuedInstantMovementEffects();
+	const TArray<FScheduledInstantMovementEffect>& ScheduledInstantMovementEffects = MoverComp.GetQueuedInstantMovementEffects();
 	FChaosNetInstantMovementEffectsQueue& InstantMovementEffectsQueue = InputData.InputCmd.InputCollection.FindOrAddMutableDataByType<FChaosNetInstantMovementEffectsQueue>();
 	InstantMovementEffectsQueue.Effects.Empty();
-	if (!ScheduledInstantMovementEffect.IsEmpty())
+	for (const FScheduledInstantMovementEffect& ScheduledEffect : MoverComp.GetQueuedInstantMovementEffects())
 	{
-		for (const FScheduledInstantMovementEffect& ScheduledEffect : MoverComp.GetQueuedInstantMovementEffects())
+		uint8 UniqueID = NextInstantMovementEffectUniqueID++;
+		if (UniqueID == 0xFF)
 		{
-			uint8 UniqueID = NextInstantMovementEffectUniqueID++;
-			if (UniqueID == 0xFF)
-			{
-				UniqueID = 0;
-				NextInstantMovementEffectUniqueID = 1;
-			}
+			UniqueID = 0;
+			NextInstantMovementEffectUniqueID = 1;
+		}
 
-			InstantMovementEffectsQueue.Add(ScheduledEffect, TimeStep.ServerFrame, UniqueID);
+		InstantMovementEffectsQueue.Add(ScheduledEffect, TimeStep.ServerFrame, UniqueID);
 
 #if !defined(BUILD_SHIPPING) || !BUILD_SHIPPING
-			ENetMode NetMode = GetWorld() ? GetWorld()->GetNetMode() : NM_MAX;
-			UE_LOG(LogChaosMover, Verbose, TEXT("(%s) UChaosMoverBackendComponent::InjectInstantMovementEffects: Transferring Instant Effect Scheduled for frame %d at frame %d (Assigning Net ID %d) from the mover component to the async simulation: %s."),
-				*ToString(NetMode), ScheduledEffect.ExecutionServerFrame, TimeStep.ServerFrame, UniqueID, ScheduledEffect.Effect.IsValid() ? *ScheduledEffect.Effect->ToSimpleString() : TEXT("INVALID INSTANT EFFECT"));
+		ENetMode NetMode = GetWorld() ? GetWorld()->GetNetMode() : NM_MAX;
+		UE_LOG(LogChaosMover, Verbose, TEXT("(%s) UChaosMoverBackendComponent::InjectInstantMovementEffects: Transferring Instant Effect Scheduled for frame %d at frame %d (Assigning Net ID %d) from the mover component to the async simulation: %s."),
+			*ToString(NetMode), ScheduledEffect.ExecutionServerFrame, TimeStep.ServerFrame, UniqueID, ScheduledEffect.Effect.IsValid() ? *ScheduledEffect.Effect->ToSimpleString() : TEXT("INVALID INSTANT EFFECT"));
 #endif // !defined(BUILD_SHIPPING) || !BUILD_SHIPPING
-		}
-		MoverComp.ClearQueuedInstantMovementEffects();
 	}
+	// The effects queue is cleared in ProduceInputData
 }
 
 void UChaosMoverBackendComponent::InjectInstantMovementEffectsIntoSim(const FMoverTimeStep& TimeStep)
@@ -588,8 +603,8 @@ void UChaosMoverBackendComponent::InjectInstantMovementEffectsIntoSim(const FMov
 			UE_LOG(LogChaosMover, Verbose, TEXT("(%s) ProduceInputData: Async simulation received an Instant Movement Effect from the Mover Component, scheduled for frame %d at frame %d: %s"),
 				*ToString(NetMode), ScheduledEffect.ExecutionServerFrame, TimeStep.ServerFrame, ScheduledEffect.Effect.IsValid() ? *ScheduledEffect.Effect->ToSimpleString() : TEXT("INVALID INSTANT EFFECT"));
 		}
-		MoverComp.ClearQueuedInstantMovementEffects();
 	}
+	// The effects queue is cleared in ProduceInputData
 }
 
 void UChaosMoverBackendComponent::ConsumeOutputData(const FMoverTimeStep& TimeStep, const UE::ChaosMover::FSimulationOutputData& OutputData)
diff --git a/Engine/Plugins/Experimental/ChaosMover/Source/ChaosMover/Private/Backends/ChaosMoverSubsystem.cpp b/Engine/Plugins/Experimental/ChaosMover/Source/ChaosMover/Private/Backends/ChaosMoverSubsystem.cpp
--- a/Engine/Plugins/Experimental/ChaosMover/Source/ChaosMover/Private/Backends/ChaosMoverSubsystem.cpp
+++ b/Engine/Plugins/Experimental/ChaosMover/Source/ChaosMover/Private/Backends/ChaosMoverSubsystem.cpp
@@ -87,6 +87,8 @@ void UChaosMoverSubsystem::InjectInputs_External(int32 PhysicsStep, int32 NumSte
 		return;
 	}
 
+	QUICK_SCOPE_CYCLE_COUNTER(ChaosMover_InjectInputs_External);
+
 	// Clear invalid data before we attempt to use it.
 	Backends.RemoveAllSwap([](const TWeakObjectPtr<UChaosMoverBackendComponent> Backend) {
 		return !Backend.IsValid();
@@ -157,6 +159,8 @@ void UChaosMoverSubsystem::OnPostPhysicsTick(FChaosScene* Scene)
 		return;
 	}
 
+	QUICK_SCOPE_CYCLE_COUNTER(ChaosMover_OnPostPhysicsTick);
+
 	const bool ForceSingleThread = UE::ChaosMover::CVars::bForceSingleThreadedGT;
 
 	// Pop each async output
@@ -169,6 +173,7 @@ void UChaosMoverSubsystem::OnPostPhysicsTick(FChaosScene* Scene)
 			}
 		};
 
+		// AsyncOutput->Backends is reset and filled each physics frame in FAsyncCallback::OnPreSimulate_Internal
 		Chaos::PhysicsParallelFor(AsyncOutput->Backends.Num(), LambdaParallelUpdate, ForceSingleThread);
 	}
 
diff --git a/Engine/Plugins/Experimental/ChaosMover/Source/ChaosMover/Private/ChaosMoverSimulation.cpp b/Engine/Plugins/Experimental/ChaosMover/Source/ChaosMover/Private/ChaosMoverSimulation.cpp
--- a/Engine/Plugins/Experimental/ChaosMover/Source/ChaosMover/Private/ChaosMoverSimulation.cpp
+++ b/Engine/Plugins/Experimental/ChaosMover/Source/ChaosMover/Private/ChaosMoverSimulation.cpp
@@ -9,6 +9,7 @@
 #include "Chaos/KinematicTargets.h"
 #include "Chaos/PBDJointConstraintData.h"
 #include "Chaos/PhysicsObjectInternalInterface.h"
+#include "Chaos/ResimConsoleVariables.h"
 #include "ChaosMover/ChaosMovementMode.h"
 #include "ChaosMover/ChaosMovementModeTransition.h"
 #include "ChaosMover/ChaosMoverConsoleVariables.h"
@@ -33,6 +34,7 @@
 #include "PhysicsEngine/PhysicsObjectExternalInterface.h"
 #include "PhysicsProxy/CharacterGroundConstraintProxy.h"
 #include "UObject/UObjectGlobals.h"
+#include "Chaos/GeometryParticles.h"
 
 #if WITH_CHAOS_VISUAL_DEBUGGER
 #include "ChaosVisualDebugger/ChaosVisualDebuggerTrace.h"
@@ -192,6 +194,8 @@ void UChaosMoverSimulation::OnDeinit()
 
 void UChaosMoverSimulation::ProcessInputs(int32 PhysicsStep, const FMoverTimeStep& TimeStep, const UE::ChaosMover::FSimulationInputData& InputData)
 {
+	QUICK_SCOPE_CYCLE_COUNTER(ChaosMoverSimulation_ProcessInputs);
+
 	Chaos::EnsureIsInPhysicsThreadContext();
 
 	// Ensure that we're using the correct inputs. On the server or during resim the inputs will
@@ -225,9 +229,14 @@ void UChaosMoverSimulation::ProcessInputs(int32 PhysicsStep, const FMoverTimeSte
 
 void UChaosMoverSimulation::SimulationTick(const FMoverTimeStep& TimeStep, const UE::ChaosMover::FSimulationInputData& InputData, UE::ChaosMover::FSimulationOutputData& OutputData)
 {
+	QUICK_SCOPE_CYCLE_COUNTER(ChaosMoverSimulation_SimulationTick);
+
 	Chaos::EnsureIsInPhysicsThreadContext();
 
-	if (bInitialized)
+	const bool bIsBubbleResim = Chaos::ResimConsoleVars::bIsBubbleResimEnabled && TimeStep.bIsResimulating;
+	const bool bShouldTick = bInitialized && (!bIsBubbleResim || ShouldResim());
+
+	if (bShouldTick)
 	{
 		RollbackBlackboard->BeginSimulationFrame(TimeStep);
 
@@ -239,8 +248,16 @@ void UChaosMoverSimulation::SimulationTick(const FMoverTimeStep& TimeStep, const
 	}
 }
 
+bool UChaosMoverSimulation::ShouldResim() const
+{
+	const Chaos::FPBDRigidParticleHandle* ControlledParticle = GetControlledParticle();
+	return ControlledParticle ? (ControlledParticle->ResimType() != Chaos::EResimType::FrozenDuringResim) : true;
+}
+
 void UChaosMoverSimulation::ModifyContacts(const FMoverTimeStep& TimeStep, const UE::ChaosMover::FSimulationInputData& InputData, const UE::ChaosMover::FSimulationOutputData& OutputData, Chaos::FCollisionContactModifier& Modifier)
 {
+	QUICK_SCOPE_CYCLE_COUNTER(ChaosMoverSimulation_ModifyContacts);
+
 	Chaos::EnsureIsInPhysicsThreadContext();
 
 	if (!bInitialized)
@@ -312,6 +329,8 @@ void UChaosMoverSimulation::ModifyContacts(const FMoverTimeStep& TimeStep, const
 
 void UChaosMoverSimulation::PreSimulationTickCharacter(const IChaosCharacterMovementModeInterface& CharacterMode, const FMoverTimeStep& TimeStep, const UE::ChaosMover::FSimulationInputData& InputData)
 {
+	QUICK_SCOPE_CYCLE_COUNTER(ChaosMoverSimulation_PreSimulationTickCharacter);
+
 	// Add inputs if we don't have them and make sure we have a valid input
 	FCharacterDefaultInputs& CharacterDefaultInputs = InputData.InputCmd.InputCollection.FindOrAddMutableDataByType<FCharacterDefaultInputs>();
 	if (CharacterDefaultInputs.GetMoveInputType() == EMoveInputType::Invalid)
@@ -387,6 +406,8 @@ void UChaosMoverSimulation::PreSimulationTickCharacter(const IChaosCharacterMove
 
 void UChaosMoverSimulation::OnSimulationRollback(const FMoverTimeStep& NewTimeStep, const FMoverSyncState& PrevSyncState)
 {
+	QUICK_SCOPE_CYCLE_COUNTER(ChaosMoverSimulation_OnSimulationRollback);
+
 	// Rollback blackboard on the first frame of resimulation
 	Blackboard->Invalidate(EInvalidationReason::Rollback);
 	RollbackBlackboard->BeginRollback(NewTimeStep);
@@ -407,6 +428,8 @@ void UChaosMoverSimulation::OnSimulationRollback(const FMoverTimeStep& NewTimeSt
 
 void UChaosMoverSimulation::OnPreSimulationTick(const FMoverTimeStep& TimeStep, const UE::ChaosMover::FSimulationInputData& InputData)
 {
+	QUICK_SCOPE_CYCLE_COUNTER(ChaosMoverSimulation_OnPreSimulationTick);
+
 	InternalServerFrame = TimeStep.ServerFrame;
 
 	// Update the sync state from the current physics state
@@ -429,6 +452,8 @@ void UChaosMoverSimulation::OnPreSimulationTick(const FMoverTimeStep& TimeStep,
 
 void UChaosMoverSimulation::OnSimulationTick(const FMoverTimeStep& TimeStep, const UE::ChaosMover::FSimulationInputData& InputData, UE::ChaosMover::FSimulationOutputData& OutputData)
 {
+	QUICK_SCOPE_CYCLE_COUNTER(ChaosMoverSimulation_OnSimulationTick);
+
 	check(Blackboard.Get());
 
 	FMoverTickStartData TickStartData(InputData.InputCmd, CurrentSyncState, InputData.AuxInputState);
@@ -443,6 +468,8 @@ void UChaosMoverSimulation::OnSimulationTick(const FMoverTimeStep& TimeStep, con
 
 void UChaosMoverSimulation::PostSimulationTickCharacter(const IChaosCharacterMovementModeInterface& CharacterMode, const FMoverTimeStep& TimeStep, UE::ChaosMover::FSimulationOutputData& OutputData)
 {
+	QUICK_SCOPE_CYCLE_COUNTER(ChaosMoverSimulation_PostSimulationTickCharacter);
+
 	FChaosMoverCharacterSimState& CharacterSimState = InternalSimData.FindOrAddMutableDataByType<FChaosMoverCharacterSimState>();
 	FMoverDefaultSyncState* PostSimDefaultSyncState = OutputData.SyncState.SyncStateCollection.FindMutableDataByType<FMoverDefaultSyncState>();
 	check(PostSimDefaultSyncState);
@@ -502,6 +529,8 @@ void UChaosMoverSimulation::PostSimulationTickCharacter(const IChaosCharacterMov
 
 void UChaosMoverSimulation::PostSimulationTickCharacterConstraint(const IChaosCharacterConstraintMovementModeInterface& CharacterConstraintMode, const FMoverTimeStep& TimeStep, UE::ChaosMover::FSimulationOutputData& OutputData)
 {
+	QUICK_SCOPE_CYCLE_COUNTER(ChaosMoverSimulation_PostSimulationTickCharacterConstraint);
+
 	if (!CharacterConstraintHandle)
 	{
 		return;
@@ -589,6 +618,8 @@ void UChaosMoverSimulation::PostSimulationTickCharacterConstraint(const IChaosCh
 
 void UChaosMoverSimulation::PostSimulationTickMovementActuation(const IChaosMovementActuationInterface& MovementActuationMode, const FMoverTimeStep& TimeStep, UE::ChaosMover::FSimulationOutputData& OutputData)
 {
+	QUICK_SCOPE_CYCLE_COUNTER(ChaosMoverSimulation_PostSimulationTickMovementActuation);
+
 	if (!ActuationConstraintHandle)
 	{
 		return;
@@ -775,6 +806,8 @@ void UChaosMoverSimulation::Teleport(const FTransform& TargetTransform, const FM
 
 void UChaosMoverSimulation::OnPostSimulationTick(const FMoverTimeStep& TimeStep, UE::ChaosMover::FSimulationOutputData& OutputData)
 {
+	QUICK_SCOPE_CYCLE_COUNTER(ChaosMoverSimulation_OnPostSimulationTick);
+
 	// TODO - make this more extensible
 	if (TStrongObjectPtr<UBaseMovementMode> CurrentModePtr = StateMachine.GetCurrentMode().Pin())
 	{
@@ -814,6 +847,8 @@ void UChaosMoverSimulation::OnPostSimulationTick(const FMoverTimeStep& TimeStep,
 
 void UChaosMoverSimulation::TraceMoverData(const FMoverTimeStep& TimeStep, const UE::ChaosMover::FSimulationOutputData& OutputData)
 {
+	QUICK_SCOPE_CYCLE_COUNTER(ChaosMoverSimulation_TraceMoverData);
+
 	// Send the latest physics thread data to CVD
 #if WITH_CHAOS_VISUAL_DEBUGGER
 	if (FChaosVisualDebuggerTrace::IsTracing())
diff --git a/Engine/Plugins/Experimental/ChaosMover/Source/ChaosMover/Public/ChaosMover/Backends/ChaosMoverBackend.h b/Engine/Plugins/Experimental/ChaosMover/Source/ChaosMover/Public/ChaosMover/Backends/ChaosMoverBackend.h
--- a/Engine/Plugins/Experimental/ChaosMover/Source/ChaosMover/Public/ChaosMover/Backends/ChaosMoverBackend.h
+++ b/Engine/Plugins/Experimental/ChaosMover/Source/ChaosMover/Public/ChaosMover/Backends/ChaosMoverBackend.h
@@ -39,6 +39,7 @@ public:
 
 	// IMoverBackendLiaisonInterface
 	CHAOSMOVER_API virtual bool IsAsync() const override;
+	CHAOSMOVER_API virtual bool ShouldResim() const override;
 	CHAOSMOVER_API virtual float GetEventSchedulingMinDelaySeconds() const override;
 	CHAOSMOVER_API virtual double GetCurrentSimTimeMs();
 	CHAOSMOVER_API virtual int32 GetCurrentSimFrame();
diff --git a/Engine/Plugins/Experimental/ChaosMover/Source/ChaosMover/Public/ChaosMover/ChaosMoverSimulation.h b/Engine/Plugins/Experimental/ChaosMover/Source/ChaosMover/Public/ChaosMover/ChaosMoverSimulation.h
--- a/Engine/Plugins/Experimental/ChaosMover/Source/ChaosMover/Public/ChaosMover/ChaosMoverSimulation.h
+++ b/Engine/Plugins/Experimental/ChaosMover/Source/ChaosMover/Public/ChaosMover/ChaosMoverSimulation.h
@@ -80,6 +80,8 @@ public:
 
 	ENetRole GetLocalRole() const { return OwnerLocalRole; }
 
+	CHAOSMOVER_API bool ShouldResim() const;
+
 	struct FInitParams
 	{
 		TMap<FName, TWeakObjectPtr<UBaseMovementMode>> ModesToRegister;
diff --git a/Engine/Plugins/Experimental/Mover/Source/Mover/Public/Backends/MoverBackendLiaison.h b/Engine/Plugins/Experimental/Mover/Source/Mover/Public/Backends/MoverBackendLiaison.h
--- a/Engine/Plugins/Experimental/Mover/Source/Mover/Public/Backends/MoverBackendLiaison.h
+++ b/Engine/Plugins/Experimental/Mover/Source/Mover/Public/Backends/MoverBackendLiaison.h
@@ -31,6 +31,8 @@ public:
 
 	// Whether this backend will simulate movement asynchronously
 	virtual bool IsAsync() const { return false; }
+	// During a resim, whether this particular mover should resimulate
+	virtual bool ShouldResim() const { return true; }
 
 	// How much delay to apply to scheduled events. This is important for networked events, and should be greater than the RTT to ensure the event will be executed on all end points at the same frame.
 	virtual float GetEventSchedulingMinDelaySeconds() const {return 0.3f;}
diff --git a/Engine/Source/Runtime/Experimental/Chaos/Private/Chaos/ResimConsoleVariables.cpp b/Engine/Source/Runtime/Experimental/Chaos/Private/Chaos/ResimConsoleVariables.cpp
--- a/Engine/Source/Runtime/Experimental/Chaos/Private/Chaos/ResimConsoleVariables.cpp
+++ b/Engine/Source/Runtime/Experimental/Chaos/Private/Chaos/ResimConsoleVariables.cpp
@@ -12,5 +12,11 @@ namespace Chaos
 			TEXT("p.Chaos.BubbleResim.Enable"),
 			bIsBubbleResimEnabled,
 			TEXT("If true bubble resim is enabled, otherwise a resim causes the whole world to resimulate, not just bubbles around server corrected particles"));
+
+		bool bCanFrozenParticlesModifyContacts = true;
+		static FAutoConsoleVariableRef CVarCanFrozenParticlesModifyContacts(
+			TEXT("p.Chaos.BubbleResim.EnableCanFrozenParticlesModifyContacts"),
+			bCanFrozenParticlesModifyContacts,
+			TEXT("Whether frozen particles are still allowed to modify contacts during a resim."));
 	}
 }
\ No newline at end of file
diff --git a/Engine/Source/Runtime/Experimental/Chaos/Public/Chaos/ResimConsoleVariables.h b/Engine/Source/Runtime/Experimental/Chaos/Public/Chaos/ResimConsoleVariables.h
--- a/Engine/Source/Runtime/Experimental/Chaos/Public/Chaos/ResimConsoleVariables.h
+++ b/Engine/Source/Runtime/Experimental/Chaos/Public/Chaos/ResimConsoleVariables.h
@@ -8,5 +8,6 @@ namespace Chaos
 	namespace ResimConsoleVars
 	{
 		extern CHAOS_API bool bIsBubbleResimEnabled;
+		extern CHAOS_API bool bCanFrozenParticlesModifyContacts;
 	}
 }
\ No newline at end of file
