Chaos Mover - Update sweep in falling mode to account for updated velocity

#rnx

diff --git a/Engine/Plugins/Experimental/ChaosMover/Source/ChaosMover/Private/Character/Modes/ChaosFallingMode.cpp b/Engine/Plugins/Experimental/ChaosMover/Source/ChaosMover/Private/Character/Modes/ChaosFallingMode.cpp
--- a/Engine/Plugins/Experimental/ChaosMover/Source/ChaosMover/Private/Character/Modes/ChaosFallingMode.cpp
+++ b/Engine/Plugins/Experimental/ChaosMover/Source/ChaosMover/Private/Character/Modes/ChaosFallingMode.cpp
@@ -118,43 +118,6 @@ void UChaosFallingMode::GenerateMove_Implementation(const FMoverTickStartData& S
 		Params.Deceleration = OverTerminalSpeedFallingDeceleration;
 	}
 
-	FFloorCheckResult FloorResult;
-	FWaterCheckResult WaterResult;
-	UMoverBlackboard* SimBlackboard = Simulation->GetBlackboard_Mutable();
-	if (SimBlackboard)
-	{
-		// Update the floor
-		UE::ChaosMover::Utils::FFloorSweepParams SweepParams{
-			.ResponseParams = DefaultSimInputs->CollisionResponseParams,
-			.QueryParams = DefaultSimInputs->CollisionQueryParams,
-			.Location = StartingSyncState->GetLocation_WorldSpace(),
-			.DeltaPos = StartingSyncState->GetVelocity_WorldSpace() * TimeStep.StepMs * 0.001f,
-			.UpDir = DefaultSimInputs->UpDir,
-			.World = DefaultSimInputs->World,
-			.QueryDistance = 1.2f * GetTargetHeight(),
-			.QueryRadius = FMath::Min(GetGroundQueryRadius(), FMath::Max(DefaultSimInputs->PawnCollisionRadius - 5.0f, 0.0f)),
-			.MaxWalkSlopeCosine = GetMaxWalkSlopeCosine(),
-			.TargetHeight = GetTargetHeight(),
-			.CollisionChannel = DefaultSimInputs->CollisionChannel
-		};
-
-		UE::ChaosMover::Utils::FloorSweep_Internal(SweepParams, FloorResult, WaterResult);
-
-		SimBlackboard->Set(CommonBlackboard::LastFloorResult, FloorResult);
-		SimBlackboard->Set(CommonBlackboard::LastWaterResult, WaterResult);
-
-		//if (FloorResult.HitResult.IsValidBlockingHit() && FloorResult.HitResult.Normal.Dot(DefaultSimInputs->UpDir) > UE::MoverUtils::VERTICAL_SLOPE_NORMAL_MAX_DOT && !FloorResult.IsWalkableFloor())
-		//{
-		//	// If acceleration is into the wall, limit contribution.
-		//	if (FVector::DotProduct(Params.MoveInput, FloorResult.HitResult.Normal) < 0.f)
-		//	{
-		//		// Allow movement parallel to the wall, but not into it because that may push us up.
-		//		const FVector FallingHitNormal = FVector::VectorPlaneProject(FloorResult.HitResult.Normal, -DefaultSimInputs->UpDir).GetSafeNormal();
-		//		Params.MoveInput = FVector::VectorPlaneProject(Params.MoveInput, FallingHitNormal);
-		//	}
-		//}
-	}
-
 	OutProposedMove = UAirMovementUtils::ComputeControlledFreeMove(Params);
 	const FVector VelocityWithGravity = StartVelocity + UMovementUtils::ComputeVelocityFromGravity(DefaultSimInputs->Gravity, DeltaSeconds);
 
@@ -180,6 +143,32 @@ void UChaosFallingMode::GenerateMove_Implementation(const FMoverTickStartData& S
 	{
 		UMovementUtils::SetGravityVerticalComponent(OutProposedMove.LinearVelocity, VelocityWithGravity.Dot(DefaultSimInputs->UpDir), DefaultSimInputs->UpDir);
 	}
+
+	FFloorCheckResult FloorResult;
+	FWaterCheckResult WaterResult;
+	UMoverBlackboard* SimBlackboard = Simulation->GetBlackboard_Mutable();
+	if (SimBlackboard)
+	{
+		// Update the floor
+		UE::ChaosMover::Utils::FFloorSweepParams SweepParams{
+			.ResponseParams = DefaultSimInputs->CollisionResponseParams,
+			.QueryParams = DefaultSimInputs->CollisionQueryParams,
+			.Location = StartingSyncState->GetLocation_WorldSpace(),
+			.DeltaPos = OutProposedMove.LinearVelocity * TimeStep.StepMs * 0.001f,
+			.UpDir = DefaultSimInputs->UpDir,
+			.World = DefaultSimInputs->World,
+			.QueryDistance = 1.5f * GetTargetHeight(),
+			.QueryRadius = FMath::Min(GetGroundQueryRadius(), FMath::Max(DefaultSimInputs->PawnCollisionRadius - 5.0f, 0.0f)),
+			.MaxWalkSlopeCosine = GetMaxWalkSlopeCosine(),
+			.TargetHeight = GetTargetHeight(),
+			.CollisionChannel = DefaultSimInputs->CollisionChannel
+		};
+
+		UE::ChaosMover::Utils::FloorSweep_Internal(SweepParams, FloorResult, WaterResult);
+
+		SimBlackboard->Set(CommonBlackboard::LastFloorResult, FloorResult);
+		SimBlackboard->Set(CommonBlackboard::LastWaterResult, WaterResult);
+	}
 }
 
 void UChaosFallingMode::SimulationTick_Implementation(const FSimulationTickParams& Params, FMoverTickEndData& OutputState)
