- Fixed a straggling Base<Plugin> -> Default<Plugin> filename

diff --git a/Engine/Plugins/Experimental/ChaosMover/Config/BaseChaosMover.ini b/Engine/Plugins/Experimental/ChaosMover/Config/BaseChaosMover.ini
deleted file mode 100644
diff --git a/Engine/Plugins/Experimental/ChaosMover/Config/DefaultChaosMover.ini b/Engine/Plugins/Experimental/ChaosMover/Config/DefaultChaosMover.ini
new file mode 100644
--- /dev/null
+++ b/Engine/Plugins/Experimental/ChaosMover/Config/DefaultChaosMover.ini
@@ -0,0 +1,2 @@
+﻿[CoreRedirects]
++PropertyRedirects=(OldName="/Script/ChaosMover.ChaosCharacterMoverComponent.OnLanded",NewName="/Script/ChaosMover.ChaosCharacterMoverComponent.OnLandedDelegate")
\ No newline at end of file
