ChaosMover : Add guard against division by 0 in proposed move linear velocity calculation.
#rnx

diff --git a/Engine/Plugins/Experimental/ChaosMover/Source/ChaosMover/Private/Character/Modes/ChaosWalkingMode.cpp b/Engine/Plugins/Experimental/ChaosMover/Source/ChaosMover/Private/Character/Modes/ChaosWalkingMode.cpp
--- a/Engine/Plugins/Experimental/ChaosMover/Source/ChaosMover/Private/Character/Modes/ChaosWalkingMode.cpp
+++ b/Engine/Plugins/Experimental/ChaosMover/Source/ChaosMover/Private/Character/Modes/ChaosWalkingMode.cpp
@@ -167,7 +167,7 @@ void UChaosWalkingMode::GenerateMove_Implementation(const FMoverTickStartData& S
 		FWaterCheckResult WaterResult;
 		GetFloorAndCheckMovement(*StartingSyncState, OutProposedMove, *DefaultSimInputs, DeltaSeconds, FloorResult, WaterResult, OutDeltaPos);
 
-		OutProposedMove.LinearVelocity = OutDeltaPos / DeltaSeconds;
+		OutProposedMove.LinearVelocity = (DeltaSeconds > 0.0f) ? (OutDeltaPos / DeltaSeconds) : FVector::ZeroVector;
 
 		SimBlackboard->Set(CommonBlackboard::LastFloorResult, FloorResult);
 		SimBlackboard->Set(CommonBlackboard::LastWaterResult, WaterResult);
@@ -175,7 +175,11 @@ void UChaosWalkingMode::GenerateMove_Implementation(const FMoverTickStartData& S
 		if (bMaintainHorizontalGroundVelocity)
 		{
 			// So far have assumed we are on level ground, so now add velocity up the slope
-			OutProposedMove.LinearVelocity -= UpDirection * OutProposedMove.LinearVelocity.Dot(FloorResult.HitResult.ImpactNormal) / UpDirection.Dot(FloorResult.HitResult.ImpactNormal);
+			const double CosAngleBetweenUpAndNormal = UpDirection.Dot(FloorResult.HitResult.ImpactNormal);
+			if (CosAngleBetweenUpAndNormal > UE_SMALL_NUMBER)
+			{
+				OutProposedMove.LinearVelocity -= UpDirection * OutProposedMove.LinearVelocity.Dot(FloorResult.HitResult.ImpactNormal) / CosAngleBetweenUpAndNormal;
+			}
 		}
 	};
 }
