[ChaosMover] Fixes for pusing into small, low objects

#rnx

diff --git a/Engine/Plugins/Experimental/ChaosMover/Source/ChaosMover/Private/Character/Modes/ChaosWalkingMode.cpp b/Engine/Plugins/Experimental/ChaosMover/Source/ChaosMover/Private/Character/Modes/ChaosWalkingMode.cpp
--- a/Engine/Plugins/Experimental/ChaosMover/Source/ChaosMover/Private/Character/Modes/ChaosWalkingMode.cpp
+++ b/Engine/Plugins/Experimental/ChaosMover/Source/ChaosMover/Private/Character/Modes/ChaosWalkingMode.cpp
@@ -146,7 +146,7 @@ void UChaosWalkingMode::GenerateMove_Implementation(const FMoverTickStartData& S
 		Params.UpDirection = UpDirection;
 		Params.bUseAccelerationForVelocityMove = SharedSettings->bUseAccelerationForVelocityMove;
 
-		if (Params.MoveInput.SizeSquared() > 0.f && !UMovementUtils::IsExceedingMaxSpeed(Params.PriorVelocity, SharedSettings->MaxSpeed))
+		if (Params.MoveInput.SizeSquared() > 0.f && !UMovementUtils::IsExceedingMaxSpeed(Params.PriorVelocity, GetMaxSpeed()))
 		{
 			Params.Friction = SharedSettings->GroundFriction;
 		}
@@ -221,10 +221,7 @@ void UChaosWalkingMode::SimulationTick_Implementation(const FSimulationTickParam
 
 	if (FloorResult.IsWalkableFloor())
 	{
-		const float InitialHeightAboveFloor = FloorResult.FloorDist - GetTargetHeight();
-
-		// Put the target position on the floor at the target height
-		FVector TargetPosition = StartingSyncState->GetLocation_WorldSpace() - UpDirection * InitialHeightAboveFloor;
+		FVector TargetPosition = StartingSyncState->GetLocation_WorldSpace();
 
 		// The base movement mode does not apply gravity in walking mode so apply here.
 		// Also remove the gravity that will be applied by the physics simulation.
@@ -280,6 +277,7 @@ void UChaosWalkingMode::SimulationTick_Implementation(const FSimulationTickParam
 		// If stepping up make sure that the step height is less than the max step height
 		// and the new surface has CanCharacterStepUpOn set to true.
 		// If stepping down make sure the step height is less than the max step height.
+		const float InitialHeightAboveFloor = FloorResult.FloorDist - GetTargetHeight();
 		const float EndHeightAboveFloor = InitialHeightAboveFloor + ProjectedRelativeVerticalVelocity * DeltaSeconds;
 		const bool bIsSteppingDown = InitialHeightAboveFloor > UE_KINDA_SMALL_NUMBER;
 		const bool bIsWithinReach = EndHeightAboveFloor <= SharedSettings->MaxStepHeight;
@@ -290,6 +288,12 @@ void UChaosWalkingMode::SimulationTick_Implementation(const FSimulationTickParam
 			TargetVelocity -= FractionalDownwardVelocityToTarget * (EndHeightAboveFloor / DeltaSeconds) * UpDirection;
 		}
 
+		// Put the target position on the floor at the target height
+		if (!bIsGroundMoving)
+		{
+			TargetPosition -= UpDirection * InitialHeightAboveFloor;
+		}
+
 		// Target orientation
 		// This is always applied regardless of whether the character is supported
 		const FRotator TargetOrientation = UMovementUtils::ApplyAngularVelocityToRotator(StartingSyncState->GetOrientation_WorldSpace(), ProposedMove.AngularVelocityDegrees, DeltaSeconds);
@@ -305,7 +309,7 @@ void UChaosWalkingMode::SimulationTick_Implementation(const FSimulationTickParam
 	OutputSyncState.MoveDirectionIntent = ProposedMove.bHasDirIntent ? ProposedMove.DirectionIntent : FVector::ZeroVector;
 }
 
-bool UChaosWalkingMode::CanStepUpOnHitSurface(const FFloorCheckResult& FloorResult) const
+bool UChaosWalkingMode::CanStepUpOnHitSurface(const FFloorCheckResult& FloorResult, const FVector& CharacterVelocity) const
 {
 	const float StepHeight = GetTargetHeight() - FloorResult.FloorDist;
 
@@ -315,6 +319,16 @@ bool UChaosWalkingMode::CanStepUpOnHitSurface(const FFloorCheckResult& FloorResu
 	if (bWalkable && SteppingUp)
 	{
 		bWalkable = UGroundMovementUtils::CanStepUpOnHitSurface(FloorResult.HitResult);
+
+		// Make sure the surface isn't moving away relative to the character movement when stepping up
+		if (const Chaos::FPBDRigidParticleHandle* GroundParticle = UChaosGroundMovementUtils::GetRigidParticleHandleFromFloorResult_Internal(FloorResult))
+		{
+			const FVector R = FloorResult.HitResult.ImpactPoint - GroundParticle->GetX();
+			const FVector GroundV = GroundParticle->GetV() + GroundParticle->GetW().Cross(R);
+			const FVector Dir = CharacterVelocity.GetSafeNormal();
+			float SurfaceV = FVector::VectorPlaneProject(CharacterVelocity - GroundV, FloorResult.HitResult.ImpactNormal).Dot(Dir);
+			bWalkable = bWalkable && (SurfaceV > UE_KINDA_SMALL_NUMBER);
+		}
 	}
 
 	return bWalkable;
@@ -350,7 +364,7 @@ void UChaosWalkingMode::GetFloorAndCheckMovement(const FMoverDefaultSyncState& S
 		return;
 	}
 
-	bool bWalkableFloor = OutFloorResult.bWalkableFloor && CanStepUpOnHitSurface(OutFloorResult);
+	bool bWalkableFloor = OutFloorResult.bWalkableFloor && CanStepUpOnHitSurface(OutFloorResult, ProposedMove.LinearVelocity);
 	if (bWalkableFloor)
 	{
 		// Walkable floor found
@@ -361,14 +375,20 @@ void UChaosWalkingMode::GetFloorAndCheckMovement(const FMoverDefaultSyncState& S
 	const float StepBlockedHeight = GetTargetHeight() - DefaultSimInputs.PawnCollisionHalfHeight + DefaultSimInputs.PawnCollisionRadius;
 	const float StepHeight = GetTargetHeight() - OutFloorResult.FloorDist;
 
-	if (StepHeight > StepBlockedHeight)
+	bool bIsDynamicSurface = false;
+	if (const Chaos::FPBDRigidParticleHandle* GroundParticle = UChaosGroundMovementUtils::GetRigidParticleHandleFromFloorResult_Internal(OutFloorResult))
+	{
+		bIsDynamicSurface = GroundParticle->IsDynamic();
+	}
+
+	if ((StepHeight > StepBlockedHeight) || bIsDynamicSurface)
 	{
 		// Collision should prevent movement. Just try to find ground at start of movement
-		SweepParams.QueryRadius = 0.75f * GetGroundQueryRadius();
+		SweepParams.QueryRadius = 0.25f * GetGroundQueryRadius();
 		SweepParams.DeltaPos = FVector::ZeroVector;
 
 		UE::ChaosMover::Utils::FloorSweep_Internal(SweepParams, OutFloorResult, OutWaterResult);
-		OutFloorResult.bWalkableFloor = OutFloorResult.bWalkableFloor && CanStepUpOnHitSurface(OutFloorResult);
+		OutFloorResult.bWalkableFloor = OutFloorResult.bWalkableFloor && CanStepUpOnHitSurface(OutFloorResult, ProposedMove.LinearVelocity);
 		return;
 	}
 
@@ -419,7 +439,7 @@ void UChaosWalkingMode::GetFloorAndCheckMovement(const FMoverDefaultSyncState& S
 		SweepParams.DeltaPos = NewDeltaPos;
 
 		UE::ChaosMover::Utils::FloorSweep_Internal(SweepParams, OutFloorResult, OutWaterResult);
-		OutFloorResult.bWalkableFloor = OutFloorResult.bWalkableFloor && CanStepUpOnHitSurface(OutFloorResult);
+		OutFloorResult.bWalkableFloor = OutFloorResult.bWalkableFloor && CanStepUpOnHitSurface(OutFloorResult, ProposedMove.LinearVelocity);
 
 		if (OutFloorResult.bWalkableFloor)
 		{
diff --git a/Engine/Plugins/Experimental/ChaosMover/Source/ChaosMover/Public/ChaosMover/Character/Modes/ChaosWalkingMode.h b/Engine/Plugins/Experimental/ChaosMover/Source/ChaosMover/Public/ChaosMover/Character/Modes/ChaosWalkingMode.h
--- a/Engine/Plugins/Experimental/ChaosMover/Source/ChaosMover/Public/ChaosMover/Character/Modes/ChaosWalkingMode.h
+++ b/Engine/Plugins/Experimental/ChaosMover/Source/ChaosMover/Public/ChaosMover/Character/Modes/ChaosWalkingMode.h
@@ -65,7 +65,7 @@ public:
 	UE_API virtual void ModifyContacts(const FMoverTimeStep& TimeStep, const UE::ChaosMover::FSimulationInputData& InputData, const UE::ChaosMover::FSimulationOutputData& OutputData, Chaos::FCollisionContactModifier& Modifier) const override;
 
 protected:
-	UE_API virtual bool CanStepUpOnHitSurface(const FFloorCheckResult& FloorResult) const;
+	UE_API virtual bool CanStepUpOnHitSurface(const FFloorCheckResult& FloorResult, const FVector& CharacterVelocity) const;
 	UE_API virtual void GetFloorAndCheckMovement(const FMoverDefaultSyncState& SyncState, const FProposedMove& ProposedMove, const FChaosMoverSimulationDefaultInputs& DefaultSimInputs, float DeltaSeconds, FFloorCheckResult& FloorResult, FWaterCheckResult& WaterResult, FVector& OutDeltaPos) const;
 };
 
