Fix PVS warnings

#rnx

diff --git a/Engine/Plugins/Experimental/ChaosMover/Source/ChaosMover/Private/ChaosMoverConsoleVariables.cpp b/Engine/Plugins/Experimental/ChaosMover/Source/ChaosMover/Private/ChaosMoverConsoleVariables.cpp
--- a/Engine/Plugins/Experimental/ChaosMover/Source/ChaosMover/Private/ChaosMoverConsoleVariables.cpp
+++ b/Engine/Plugins/Experimental/ChaosMover/Source/ChaosMover/Private/ChaosMoverConsoleVariables.cpp
@@ -46,7 +46,7 @@ namespace UE::ChaosMover
 			FConsoleCommandWithWorldDelegate::CreateStatic([](UWorld* World)
 				{
 					APlayerController* LocalPlayerController = World ? World->GetFirstPlayerController() : nullptr;
-					APawn* LocalPawn = LocalPlayerController ? LocalPlayerController->GetPawn() : nullptr;
+					APawn* LocalPawn = LocalPlayerController ? LocalPlayerController->GetPawn().Get() : nullptr;
 					UMoverComponent* MoverComponent = LocalPawn ? LocalPawn->FindComponentByClass<UMoverComponent>() : nullptr;
 					if (!MoverComponent)
 					{
