Mover, ChaosMover: adding more contextual information to time steps and events

diff --git a/Engine/Plugins/Experimental/ChaosMover/Source/ChaosMover/Private/Backends/ChaosMoverAsyncCallback.cpp b/Engine/Plugins/Experimental/ChaosMover/Source/ChaosMover/Private/Backends/ChaosMoverAsyncCallback.cpp
--- a/Engine/Plugins/Experimental/ChaosMover/Source/ChaosMover/Private/Backends/ChaosMoverAsyncCallback.cpp
+++ b/Engine/Plugins/Experimental/ChaosMover/Source/ChaosMover/Private/Backends/ChaosMoverAsyncCallback.cpp
@@ -125,6 +125,7 @@ namespace UE::ChaosMover
 		{
 			OutMoverTimeStep.ServerFrame = AsyncInput->PhysicsSolver->GetCurrentFrame() + AsyncInput->NetworkPhysicsTickOffset;
 			OutMoverTimeStep.bIsResimulating = AsyncInput->PhysicsSolver->GetEvolution()->IsResimming();
+			OutMoverTimeStep.bIsFirstResimFrame = OutMoverTimeStep.bIsResimulating && (AsyncInput->PhysicsSolver->GetCurrentFrame() == AsyncInput->PhysicsSolver->GetRewindData()->GetResimFrame());
 		}
 	}
 
diff --git a/Engine/Plugins/Experimental/ChaosMover/Source/ChaosMover/Private/ChaosMoverStateMachine.cpp b/Engine/Plugins/Experimental/ChaosMover/Source/ChaosMover/Private/ChaosMoverStateMachine.cpp
--- a/Engine/Plugins/Experimental/ChaosMover/Source/ChaosMover/Private/ChaosMoverStateMachine.cpp
+++ b/Engine/Plugins/Experimental/ChaosMover/Source/ChaosMover/Private/ChaosMoverStateMachine.cpp
@@ -126,6 +126,11 @@ namespace UE::ChaosMover
 	}
 
 	void FMoverStateMachine::QueueNextMode(FName DesiredNextModeName, bool bShouldReenter)
+	{
+		QueueNextModeInternal(DesiredNextModeName, bShouldReenter, /*bIsFromRollback*/ false);
+	}
+
+	void FMoverStateMachine::QueueNextModeInternal(FName DesiredNextModeName, bool bShouldReenter, bool bIsFromRollback)
 	{
 		TStrongObjectPtr<UImmediateMovementModeTransition> QueuedModeTransition = QueuedModeTransitionWeakPtr.Pin();
 		if (QueuedModeTransition)
@@ -143,7 +148,7 @@ namespace UE::ChaosMover
 
 				if (Modes.Contains(DesiredNextModeName))
 				{
-					QueuedModeTransition->SetNextMode(DesiredNextModeName, bShouldReenter);
+					QueuedModeTransition->SetNextMode(DesiredNextModeName, bShouldReenter, bIsFromRollback);
 				}
 				else
 				{
@@ -155,7 +160,7 @@ namespace UE::ChaosMover
 
 	void FMoverStateMachine::SetModeImmediately(FName DesiredModeName, bool bShouldReenter)
 	{
-		QueueNextMode(DesiredModeName, bShouldReenter);
+		QueueNextModeInternal(DesiredModeName, bShouldReenter, /*bIsFromRollback*/ false);
 		AdvanceToNextMode();
 	}
 
@@ -636,7 +641,7 @@ namespace UE::ChaosMover
 
 		if (CurrentModeName != NewSyncState.MovementMode)
 		{
-			QueueNextMode(NewSyncState.MovementMode);
+			QueueNextModeInternal(NewSyncState.MovementMode, /*bShouldReenter*/ false, /*bIsFromRollback*/ true);
 		}
 
 		QueuedLayeredMoves.Empty();
@@ -809,7 +814,10 @@ namespace UE::ChaosMover
 						// Notify the simulation of a mode change so it can react accordingly
 						if (TStrongObjectPtr<UChaosMoverSimulation> SimStrongObjPtr = Simulation.Pin())
 						{
-							SimStrongObjPtr->AddEvent(MakeShared<FMovementModeChangedEventData>(InternalSimTimeMs, PreviousModeName, NextModeName));
+							TSharedPtr<FMovementModeChangedEventData> ModeChangedEventPtr = MakeShared<FMovementModeChangedEventData>(InternalSimTimeMs, PreviousModeName, NextModeName);
+							ModeChangedEventPtr->Context.bIsCausedByRollback = QueuedModeTransition->IsCausedByRollback();
+
+							SimStrongObjPtr->AddEvent(ModeChangedEventPtr);
 						}
 					}
 				}
diff --git a/Engine/Plugins/Experimental/ChaosMover/Source/ChaosMover/Private/Character/Transitions/ChaosCharacterLandingCheck.cpp b/Engine/Plugins/Experimental/ChaosMover/Source/ChaosMover/Private/Character/Transitions/ChaosCharacterLandingCheck.cpp
--- a/Engine/Plugins/Experimental/ChaosMover/Source/ChaosMover/Private/Character/Transitions/ChaosCharacterLandingCheck.cpp
+++ b/Engine/Plugins/Experimental/ChaosMover/Source/ChaosMover/Private/Character/Transitions/ChaosCharacterLandingCheck.cpp
@@ -112,7 +112,7 @@ TSharedPtr<FLandedEventData> UChaosCharacterLandingCheck::AllocLandedEventData()
 
 void UChaosCharacterLandingCheck::PopulateLandedEventData(const FSimulationTickParams& Params, const FFloorCheckResult& FloorResult, FLandedEventData& OutLandedEvent)
 {
-	OutLandedEvent.EventTimeMs = Params.TimeStep.BaseSimTimeMs;
+	OutLandedEvent.Context.EventTimeMs = Params.TimeStep.BaseSimTimeMs;
 	OutLandedEvent.HitResult = FloorResult.HitResult;
 	OutLandedEvent.NewModeName = TransitionToGroundMode;
 }
diff --git a/Engine/Plugins/Experimental/ChaosMover/Source/ChaosMover/Public/ChaosMover/ChaosMoverStateMachine.h b/Engine/Plugins/Experimental/ChaosMover/Source/ChaosMover/Public/ChaosMover/ChaosMoverStateMachine.h
--- a/Engine/Plugins/Experimental/ChaosMover/Source/ChaosMover/Public/ChaosMover/ChaosMoverStateMachine.h
+++ b/Engine/Plugins/Experimental/ChaosMover/Source/ChaosMover/Public/ChaosMover/ChaosMoverStateMachine.h
@@ -74,6 +74,9 @@ namespace UE::ChaosMover
 
 	protected:
 
+		CHAOSMOVER_API void QueueNextModeInternal(FName DesiredNextModeName, bool bShouldReenter = false, bool bIsFromRollback=false);
+
+
 		TMap<FName, TWeakObjectPtr<UBaseMovementMode>> Modes;
 		TArray<TWeakObjectPtr<UBaseMovementModeTransition>> GlobalTransitions;
 
diff --git a/Engine/Plugins/Experimental/Mover/Source/Mover/Private/Backends/MoverNetworkPhysicsLiaisonBase.cpp b/Engine/Plugins/Experimental/Mover/Source/Mover/Private/Backends/MoverNetworkPhysicsLiaisonBase.cpp
--- a/Engine/Plugins/Experimental/Mover/Source/Mover/Private/Backends/MoverNetworkPhysicsLiaisonBase.cpp
+++ b/Engine/Plugins/Experimental/Mover/Source/Mover/Private/Backends/MoverNetworkPhysicsLiaisonBase.cpp
@@ -823,7 +823,7 @@ void UMoverNetworkPhysicsLiaisonComponentBase::PerformPreSimulate_Internal(const
 	const FMoverTimeStep TimeStep = bUsingAsyncPhysics ? GetCurrentAsyncMoverTimeStep_Internal() : GetCurrentMoverTimeStep(TickParams.DeltaTimeSeconds);
 
 	// Update movement state machine
-	if (MoverComp.bHasRolledBack)
+	if (MoverComp.bHasJustRolledBack)
 	{
 		MoverComp.ProcessFirstSimTickAfterRollback(TimeStep);
 	}
diff --git a/Engine/Plugins/Experimental/Mover/Source/Mover/Private/MovementMode.cpp b/Engine/Plugins/Experimental/Mover/Source/Mover/Private/MovementMode.cpp
--- a/Engine/Plugins/Experimental/Mover/Source/Mover/Private/MovementMode.cpp
+++ b/Engine/Plugins/Experimental/Mover/Source/Mover/Private/MovementMode.cpp
@@ -77,7 +77,7 @@ void UBaseMovementMode::Deactivate()
 	}
 }
 
-void UBaseMovementMode::Activate_External()
+void UBaseMovementMode::Activate_External(const FMoverEventContext& Context)
 {
 	if (bSupportsAsync)
 	{
@@ -85,7 +85,7 @@ void UBaseMovementMode::Activate_External()
 	}
 }
 
-void UBaseMovementMode::Deactivate_External()
+void UBaseMovementMode::Deactivate_External(const FMoverEventContext& Context)
 {
 	if (bSupportsAsync)
 	{
diff --git a/Engine/Plugins/Experimental/Mover/Source/Mover/Private/MovementModeStateMachine.cpp b/Engine/Plugins/Experimental/Mover/Source/Mover/Private/MovementModeStateMachine.cpp
--- a/Engine/Plugins/Experimental/Mover/Source/Mover/Private/MovementModeStateMachine.cpp
+++ b/Engine/Plugins/Experimental/Mover/Source/Mover/Private/MovementModeStateMachine.cpp
@@ -118,6 +118,11 @@ void UMovementModeStateMachine::ClearAllGlobalTransitions()
 }
 
 void UMovementModeStateMachine::QueueNextMode(FName DesiredNextModeName, bool bShouldReenter)
+{
+	QueueNextModeInternal(DesiredNextModeName, bShouldReenter, /*bIsFromRollback*/ false);
+}
+
+void UMovementModeStateMachine::QueueNextModeInternal(FName DesiredNextModeName, bool bShouldReenter, bool bIsFromRollback)
 {
 	if (DesiredNextModeName != NAME_None)
 	{ 
@@ -132,7 +137,7 @@ void UMovementModeStateMachine::QueueNextMode(FName DesiredNextModeName, bool bS
 
 		if (Modes.Contains(DesiredNextModeName))
 		{
-			QueuedModeTransition->SetNextMode(DesiredNextModeName, bShouldReenter);
+			QueuedModeTransition->SetNextMode(DesiredNextModeName, bShouldReenter, bIsFromRollback);
 		}
 		else
 		{
@@ -144,7 +149,13 @@ void UMovementModeStateMachine::QueueNextMode(FName DesiredNextModeName, bool bS
 
 void UMovementModeStateMachine::SetModeImmediately(FName DesiredModeName, bool bShouldReenter)
 {
-	QueueNextMode(DesiredModeName, bShouldReenter);
+	SetModeImmediatelyInternal(DesiredModeName, bShouldReenter, /*bIsFromRollback*/ false);
+}
+
+
+void UMovementModeStateMachine::SetModeImmediatelyInternal(FName DesiredModeName, bool bShouldReenter, bool bIsFromRollback)
+{
+	QueueNextModeInternal(DesiredModeName, bShouldReenter, bIsFromRollback);
 	AdvanceToNextMode();
 }
 
@@ -480,7 +491,7 @@ void UMovementModeStateMachine::OnSimulationRollback(const FMoverSyncState* Sync
 
 	if (CurrentModeName != SyncState->MovementMode)
 	{
-		SetModeImmediately(SyncState->MovementMode);
+		SetModeImmediatelyInternal(SyncState->MovementMode, /*bShouldReenter*/ false, /*bIsFromRollback*/ true);
 	}
 
 	{
diff --git a/Engine/Plugins/Experimental/Mover/Source/Mover/Private/MovementModeTransition.cpp b/Engine/Plugins/Experimental/Mover/Source/Mover/Private/MovementModeTransition.cpp
--- a/Engine/Plugins/Experimental/Mover/Source/Mover/Private/MovementModeTransition.cpp
+++ b/Engine/Plugins/Experimental/Mover/Source/Mover/Private/MovementModeTransition.cpp
@@ -80,14 +80,16 @@ void UImmediateMovementModeTransition::Trigger_Implementation(const FSimulationT
 	Clear();
 }
 
-void UImmediateMovementModeTransition::SetNextMode(FName DesiredModeName, bool bShouldReenter)
+void UImmediateMovementModeTransition::SetNextMode(FName DesiredModeName, bool bShouldReenter, bool bIsFromRollback)
 {
 	NextMode = DesiredModeName;
 	bAllowModeReentry = bShouldReenter;
+	bIsCausedByRollback = bIsFromRollback;
 }
 
 void UImmediateMovementModeTransition::Clear()
 {
 	NextMode = NAME_None;
 	bAllowModeReentry = false;
+	bIsCausedByRollback = false;
 }
diff --git a/Engine/Plugins/Experimental/Mover/Source/Mover/Private/MoverComponent.cpp b/Engine/Plugins/Experimental/Mover/Source/Mover/Private/MoverComponent.cpp
--- a/Engine/Plugins/Experimental/Mover/Source/Mover/Private/MoverComponent.cpp
+++ b/Engine/Plugins/Experimental/Mover/Source/Mover/Private/MoverComponent.cpp
@@ -529,13 +529,15 @@ void UMoverComponent::SimulationTick(const FMoverTimeStep& InTimeStep, const FMo
 	// Send mover info to the Chaos Visual Debugger (this will do nothing if CVD is not recording, or the mover info data channel not enabled)
 	UE::MoverUtils::FMoverCVDRuntimeTrace::TraceMoverData(this, &SimInput.InputCmd, &SimInput.SyncState);
 
+	// TODO: Determination of bIsResimulating and bIsFirstResimFrame should be the responsibility of the liaison / backend
 	const bool bIsResimulating = InTimeStep.BaseSimTimeMs <= CachedNewestSimTickTimeStep.BaseSimTimeMs;
 
 	FMoverTimeStep MoverTimeStep(InTimeStep);
 	MoverTimeStep.bIsResimulating = bIsResimulating;
 
-	if (bHasRolledBack)
+	if (bHasJustRolledBack)
 	{
+		MoverTimeStep.bIsFirstResimFrame = MoverTimeStep.bIsResimulating;
 		ProcessFirstSimTickAfterRollback(InTimeStep);
 	}
 
@@ -1128,14 +1130,14 @@ void UMoverComponent::OnSimulationRollback(const FMoverSyncState* SyncState, con
 	ModeFSM->OnSimulationRollback(SyncState, AuxState, NewBaseTimeStep);
 
 	RollbackBlackboard_InternalWrapper->EndRollback();
-	bHasRolledBack = true;
+	bHasJustRolledBack = true;
 }
 
 
 void UMoverComponent::ProcessFirstSimTickAfterRollback(const FMoverTimeStep& TimeStep)
 {
 	OnPostSimulationRollback.Broadcast(TimeStep, CachedLastSimTickTimeStep);
-	bHasRolledBack = false;
+	bHasJustRolledBack = false;
 }
 
 
@@ -2706,7 +2708,7 @@ void UMoverComponent::ProcessSimulationEvent(const FMoverSimulationEventData& Ev
 				UBaseMovementMode* PrevMode = PrevModePtr->Get();
 				if (PrevMode && PrevMode->bSupportsAsync)
 				{
-					PrevMode->Deactivate_External();
+					PrevMode->Deactivate_External(EventData.Context);
 				}
 			}
 		}
@@ -2718,7 +2720,7 @@ void UMoverComponent::ProcessSimulationEvent(const FMoverSimulationEventData& Ev
 				UBaseMovementMode* NewMode = NewModePtr->Get();
 				if (NewMode && NewMode->bSupportsAsync)
 				{
-					NewMode->Activate_External();
+					NewMode->Activate_External(EventData.Context);
 				}
 			}
 		}
diff --git a/Engine/Plugins/Experimental/Mover/Source/Mover/Private/MoverSimulationTypes.cpp b/Engine/Plugins/Experimental/Mover/Source/Mover/Private/MoverSimulationTypes.cpp
--- a/Engine/Plugins/Experimental/Mover/Source/Mover/Private/MoverSimulationTypes.cpp
+++ b/Engine/Plugins/Experimental/Mover/Source/Mover/Private/MoverSimulationTypes.cpp
@@ -100,7 +100,7 @@ namespace UE::Mover
 
 			// If resimulating find any events stored after the time step and remove
 			Events.RemoveAll([&InTimeStep](const TSharedPtr<FMoverSimulationEventData> Event) {
-				return !Event.IsValid() || (Event->EventTimeMs >= InTimeStep.BaseSimTimeMs);
+				return !Event.IsValid() || (Event->Context.EventTimeMs >= InTimeStep.BaseSimTimeMs);
 				});
 		}
 
@@ -144,7 +144,7 @@ namespace UE::Mover
 
 		for (TSharedPtr<FMoverSimulationEventData>& Event : Events)
 		{
-			if (Event->EventTimeMs <= AtBaseTimeMs)
+			if (Event->Context.EventTimeMs <= AtBaseTimeMs)
 			{
 				OutData.Events.Add(MoveTemp(Event));
 			}
diff --git a/Engine/Plugins/Experimental/Mover/Source/Mover/Public/MovementMode.h b/Engine/Plugins/Experimental/Mover/Source/Mover/Public/MovementMode.h
--- a/Engine/Plugins/Experimental/Mover/Source/Mover/Public/MovementMode.h
+++ b/Engine/Plugins/Experimental/Mover/Source/Mover/Public/MovementMode.h
@@ -52,8 +52,8 @@ public:
 
 	// These functions are called when the sync state is changed on the game thread
 	// and a new mode is activated/deactivated
-	UE_API virtual void Activate_External();
-	UE_API virtual void Deactivate_External();
+	UE_API virtual void Activate_External(const FMoverEventContext& Context);
+	UE_API virtual void Deactivate_External(const FMoverEventContext& Context);
 	
 	UFUNCTION(BlueprintNativeEvent, meta = (DisplayName = "Generate Move", ForceAsFunction))
 	UE_API void GenerateMove(const FMoverTickStartData& StartState, const FMoverTimeStep& TimeStep, UPARAM(ref) FProposedMove& OutProposedMove) const;
diff --git a/Engine/Plugins/Experimental/Mover/Source/Mover/Public/MovementModeStateMachine.h b/Engine/Plugins/Experimental/Mover/Source/Mover/Public/MovementModeStateMachine.h
--- a/Engine/Plugins/Experimental/Mover/Source/Mover/Public/MovementModeStateMachine.h
+++ b/Engine/Plugins/Experimental/Mover/Source/Mover/Public/MovementModeStateMachine.h
@@ -76,6 +76,9 @@ protected:
 	void ProcessEvents(const TArray<TSharedPtr<FMoverSimulationEventData>>& InEvents);
 	UE_API virtual void ProcessSimulationEvent(const FMoverSimulationEventData& EventData);
 
+	UE_API void QueueNextModeInternal(FName DesiredNextModeName, bool bShouldReenter=false, bool bIsFromRollback=false);
+	UE_API void SetModeImmediatelyInternal(FName DesiredModeName, bool bShouldReenter=false, bool bIsFromRollback=false);
+
 	UE_API virtual void PostInitProperties() override;
 
 	UPROPERTY()
diff --git a/Engine/Plugins/Experimental/Mover/Source/Mover/Public/MovementModeTransition.h b/Engine/Plugins/Experimental/Mover/Source/Mover/Public/MovementModeTransition.h
--- a/Engine/Plugins/Experimental/Mover/Source/Mover/Public/MovementModeTransition.h
+++ b/Engine/Plugins/Experimental/Mover/Source/Mover/Public/MovementModeTransition.h
@@ -77,20 +77,20 @@ public:
 
 	/** Whether this transition should reenter a mode if it evaluates true and wants to transition into a mode the actor is already in */
 	UPROPERTY(EditDefaultsOnly, BlueprintReadWrite, Category = Mover)
-	bool bAllowModeReentry = false;
+	uint8 bAllowModeReentry : 1 = 0;
 
 	/** Whether this transition should only apply to the first step of the update. If true, modes reached after transitions or mode changes
 	* in the current update will not consider this transition
 	*/
 	UPROPERTY(EditDefaultsOnly, BlueprintReadWrite, Category = Mover)
-	bool bFirstSubStepOnly = false;
+	uint8 bFirstSubStepOnly : 1 = 0;
 
 	/** 
 	 * Whether this movement mode transition supports being part of an asynchronous movement simulation (running concurrently with the gameplay thread) 
 	 * Specifically for the Evaluate and Trigger functions
 	 */
 	UPROPERTY(EditDefaultsOnly, BlueprintReadOnly, Category = Mover)
-	bool bSupportsAsync = false;
+	uint8 bSupportsAsync : 1 = 0;
 
 protected:
 	UFUNCTION(BlueprintImplementableEvent, Category = MovementModeTransition, meta = (DisplayName = "On Registered", ScriptName = "OnRegistered"))
@@ -118,14 +118,16 @@ public:
 	UE_API virtual void Trigger_Implementation(const FSimulationTickParams& Params) override;
 
 	bool IsSet() const { return !NextMode.IsNone(); }
-	UE_API void SetNextMode(FName DesiredModeName, bool bShouldReenter = false);
+	UE_API void SetNextMode(FName DesiredModeName, bool bShouldReenter=false, bool bIsFromRollback=false);
 	UE_API void Clear();
 
 	FName GetNextModeName() const { return NextMode; }
 	bool ShouldReenter() const { return bAllowModeReentry; }
+	bool IsCausedByRollback() const { return bIsCausedByRollback; }
 
 private:
 	FName NextMode = NAME_None;
+	uint8 bIsCausedByRollback : 1 = 0;
 };
 
 
diff --git a/Engine/Plugins/Experimental/Mover/Source/Mover/Public/MoverComponent.h b/Engine/Plugins/Experimental/Mover/Source/Mover/Public/MoverComponent.h
--- a/Engine/Plugins/Experimental/Mover/Source/Mover/Public/MoverComponent.h
+++ b/Engine/Plugins/Experimental/Mover/Source/Mover/Public/MoverComponent.h
@@ -889,7 +889,7 @@ private:
 	bool bInInitializeComponent = false;
 
 	// Transient flag indicating we've had a rollback and haven't started simulating forward again yet
-	bool bHasRolledBack = false;
+	bool bHasJustRolledBack = false;
 
 	/**
 	 * A cached quaternion representing the rotation from world space to gravity relative space defined by GravityAccelOverride.
diff --git a/Engine/Plugins/Experimental/Mover/Source/Mover/Public/MoverSimulationTypes.h b/Engine/Plugins/Experimental/Mover/Source/Mover/Public/MoverSimulationTypes.h
--- a/Engine/Plugins/Experimental/Mover/Source/Mover/Public/MoverSimulationTypes.h
+++ b/Engine/Plugins/Experimental/Mover/Source/Mover/Public/MoverSimulationTypes.h
@@ -517,8 +517,8 @@ struct FMoverSimulationEventData
 
 	FMoverSimulationEventData(double InEventTimeMs, FEventProcessedCallbackPtr InEventProcessedCallback = nullptr)
 		: EventProcessedCallback(InEventProcessedCallback)
-		, EventTimeMs(InEventTimeMs)
 	{
+		Context.EventTimeMs = InEventTimeMs;
 	}
 	FMoverSimulationEventData() {}
 	virtual ~FMoverSimulationEventData() {}
@@ -557,7 +557,7 @@ private:
 	FEventProcessedCallbackPtr EventProcessedCallback = nullptr;
 
 public:
-	double EventTimeMs = 0.0;
+	FMoverEventContext Context;
 };
 
 USTRUCT()
diff --git a/Engine/Plugins/Experimental/Mover/Source/Mover/Public/MoverTypes.h b/Engine/Plugins/Experimental/Mover/Source/Mover/Public/MoverTypes.h
--- a/Engine/Plugins/Experimental/Mover/Source/Mover/Public/MoverTypes.h
+++ b/Engine/Plugins/Experimental/Mover/Source/Mover/Public/MoverTypes.h
@@ -81,6 +81,20 @@ struct FMoverOnImpactParams
 	FVector AttemptedMoveDelta;
 };
 
+// Struct containing additional context about an event occurrence, such as what led to it
+USTRUCT(BlueprintType)
+struct FMoverEventContext
+{
+	GENERATED_BODY()
+
+	// Simulation time that this event occurred (in server simulation timespace)
+	UPROPERTY(EditAnywhere, BlueprintReadWrite, Category = Mover)
+	double EventTimeMs = -1.0;
+
+	UPROPERTY(EditAnywhere, BlueprintReadWrite, Category = Mover)
+	uint8 bIsCausedByRollback : 1 = 0;
+};
+
 
 USTRUCT(BlueprintType)
 struct FMoverTimeStep
@@ -101,8 +115,12 @@ struct FMoverTimeStep
 	float StepMs=1.0;
 
 	// Indicates whether this time step is re-simulating based on prior inputs, such as during a correction
+	UPROPERTY(EditAnywhere, BlueprintReadWrite, Category=Mover)
 	uint8 bIsResimulating : 1 = 0;
 
+	// If resimulating, is this the first frame after the rollback?
+	UPROPERTY(EditAnywhere, BlueprintReadWrite, Category=Mover)
+	uint8 bIsFirstResimFrame : 1 = 0;
 };
 
 
