[Chaos Mover] Fix overriding of query radius

diff --git a/Engine/Plugins/Experimental/ChaosMover/Source/ChaosMover/Private/Character/Modes/ChaosCharacterMovementMode.cpp b/Engine/Plugins/Experimental/ChaosMover/Source/ChaosMover/Private/Character/Modes/ChaosCharacterMovementMode.cpp
--- a/Engine/Plugins/Experimental/ChaosMover/Source/ChaosMover/Private/Character/Modes/ChaosCharacterMovementMode.cpp
+++ b/Engine/Plugins/Experimental/ChaosMover/Source/ChaosMover/Private/Character/Modes/ChaosCharacterMovementMode.cpp
@@ -34,6 +34,11 @@ void UChaosCharacterMovementMode::OnRegistered(const FName ModeName)
 		TargetHeight = -CharacterOwner->GetMesh()->GetRelativeLocation().Z;
 	}
 
+	if (QueryRadiusOverride.IsSet())
+	{
+		QueryRadius = QueryRadiusOverride.GetValue();
+	}
+
 	SharedSettings = GetMoverComponent()->FindSharedSettings<USharedChaosCharacterMovementSettings>();
 	ensureMsgf(SharedSettings, TEXT("Failed to find instance of USharedChaosCharacterMovementSettings on %s. Movement may not function properly."), *GetPathNameSafe(this));
 }
