Chaos Mover events: adding more time step information, like whether the step is a resimulation.  Plumbing this all the way through the simulation events system

diff --git a/Engine/Plugins/Experimental/ChaosMover/Source/ChaosMover/Private/ChaosMoverSimulation.cpp b/Engine/Plugins/Experimental/ChaosMover/Source/ChaosMover/Private/ChaosMoverSimulation.cpp
--- a/Engine/Plugins/Experimental/ChaosMover/Source/ChaosMover/Private/ChaosMoverSimulation.cpp
+++ b/Engine/Plugins/Experimental/ChaosMover/Source/ChaosMover/Private/ChaosMoverSimulation.cpp
@@ -1224,7 +1224,7 @@ void UChaosMoverSimulation::AttemptTeleport(const FMoverTimeStep& TimeStep, cons
 		// Then we need to teleport
 		Teleport(TeleportTransform, TimeStep, OutputState);
 
-		AddEvent(MakeShared<FTeleportSucceededEventData>(TimeStep.BaseSimTimeMs, FromTransform.GetLocation(), FromTransform.GetRotation(), ToLocation, ToRotation));
+		AddEvent(MakeShared<FTeleportSucceededEventData>(TimeStep, FromTransform.GetLocation(), FromTransform.GetRotation(), ToLocation, ToRotation));
 
 #if !defined(BUILD_SHIPPING) || !BUILD_SHIPPING
 		const FChaosMoverSimulationDefaultInputs* SimInputs = LocalSimInput.FindDataByType<FChaosMoverSimulationDefaultInputs>();
@@ -1239,7 +1239,7 @@ void UChaosMoverSimulation::AttemptTeleport(const FMoverTimeStep& TimeStep, cons
 	else
 	{
 		ETeleportFailureReason FailureReason = ETeleportFailureReason::Reason_NotAvailable;
-		AddEvent(MakeShared<FTeleportFailedEventData>(TimeStep.BaseSimTimeMs, FromTransform.GetLocation(), FromTransform.GetRotation(), ToLocation, ToRotation, FailureReason));
+		AddEvent(MakeShared<FTeleportFailedEventData>(TimeStep, FromTransform.GetLocation(), FromTransform.GetRotation(), ToLocation, ToRotation, FailureReason));
 
 #if !defined(BUILD_SHIPPING) || !BUILD_SHIPPING
 		const FChaosMoverSimulationDefaultInputs* SimInputs = LocalSimInput.FindDataByType<FChaosMoverSimulationDefaultInputs>();
diff --git a/Engine/Plugins/Experimental/ChaosMover/Source/ChaosMover/Private/ChaosMoverStateMachine.cpp b/Engine/Plugins/Experimental/ChaosMover/Source/ChaosMover/Private/ChaosMoverStateMachine.cpp
--- a/Engine/Plugins/Experimental/ChaosMover/Source/ChaosMover/Private/ChaosMoverStateMachine.cpp
+++ b/Engine/Plugins/Experimental/ChaosMover/Source/ChaosMover/Private/ChaosMoverStateMachine.cpp
@@ -324,8 +324,7 @@ namespace UE::ChaosMover
 			return;
 		}
 
-		InternalSimTimeMs = TimeStep.BaseSimTimeMs;
-		InternalServerFrame = TimeStep.ServerFrame;
+		InternalSimTime = TimeStep;
 
 		// Copy queued instant movement effects from the input to the state machine
 		// After this, the state machine works on its own queue, into which it can enqueue instant movement effects while stepping
@@ -349,7 +348,7 @@ namespace UE::ChaosMover
 		float TotalUsedMs = 0.0f;
 		while (TotalUsedMs < TimeStep.StepMs)
 		{
-			InternalSimTimeMs = SubTimeStep.BaseSimTimeMs;
+			InternalSimTime.BaseSimTimeMs = SubTimeStep.BaseSimTimeMs;
 
 			if (!bIsWorkingStartStateReady)
 			{
@@ -595,7 +594,7 @@ namespace UE::ChaosMover
 			++SubStepCount;
 		}
 
-		InternalSimTimeMs = TimeStep.BaseSimTimeMs + TotalUsedMs;
+		InternalSimTime.BaseSimTimeMs = TimeStep.BaseSimTimeMs + TotalUsedMs;
 
 		if (HasAnyInstantEffectsQueued(TimeStep))
 		{
@@ -632,10 +631,8 @@ namespace UE::ChaosMover
 
 	void FMoverStateMachine::OnSimulationRollback(const FMoverTimeStep& NewTimeStep, const FMoverSyncState& InvalidSyncState, const FMoverSyncState& NewSyncState)
 	{
-		float PreviousInternalSimTimeMs = InternalSimTimeMs;
-		int32 PreviousInternalServerFrame = InternalServerFrame;
-		InternalSimTimeMs = NewTimeStep.BaseSimTimeMs;
-		InternalServerFrame = NewTimeStep.ServerFrame;
+		const int32 PreviousInternalServerFrame = InternalSimTime.ServerFrame;
+		InternalSimTime = NewTimeStep;
 
 		ClearQueuedMode();
 
@@ -803,8 +800,8 @@ namespace UE::ChaosMover
 							FMovementModeChangeRecord ModeChangeRecord;
 							ModeChangeRecord.ModeName = CurrentModeName;
 							ModeChangeRecord.PrevModeName = PreviousModeName;
-							ModeChangeRecord.Frame = InternalServerFrame;
-							ModeChangeRecord.SimTimeMs = InternalSimTimeMs;
+							ModeChangeRecord.Frame = InternalSimTime.ServerFrame;
+							ModeChangeRecord.SimTimeMs = InternalSimTime.BaseSimTimeMs;
 
 							RollbackBlackboard->TrySet(CommonBlackboard::LastModeChangeRecord, ModeChangeRecord);
 						}
@@ -814,7 +811,7 @@ namespace UE::ChaosMover
 						// Notify the simulation of a mode change so it can react accordingly
 						if (TStrongObjectPtr<UChaosMoverSimulation> SimStrongObjPtr = Simulation.Pin())
 						{
-							TSharedPtr<FMovementModeChangedEventData> ModeChangedEventPtr = MakeShared<FMovementModeChangedEventData>(InternalSimTimeMs, PreviousModeName, NextModeName);
+							TSharedPtr<FMovementModeChangedEventData> ModeChangedEventPtr = MakeShared<FMovementModeChangedEventData>(InternalSimTime, PreviousModeName, NextModeName);
 							ModeChangedEventPtr->Context.bIsCausedByRollback = QueuedModeTransition->IsCausedByRollback();
 
 							SimStrongObjPtr->AddEvent(ModeChangedEventPtr);
diff --git a/Engine/Plugins/Experimental/ChaosMover/Source/ChaosMover/Private/Character/Modifiers/ChaosStanceModifier.cpp b/Engine/Plugins/Experimental/ChaosMover/Source/ChaosMover/Private/Character/Modifiers/ChaosStanceModifier.cpp
--- a/Engine/Plugins/Experimental/ChaosMover/Source/ChaosMover/Private/Character/Modifiers/ChaosStanceModifier.cpp
+++ b/Engine/Plugins/Experimental/ChaosMover/Source/ChaosMover/Private/Character/Modifiers/ChaosStanceModifier.cpp
@@ -176,7 +176,7 @@ void FChaosStanceModifier::UpdateStance(const FMovementModifierParams_Async& Asy
 	};
 
 	// Fire simulation event
-	Sim->AddEvent(MakeShared<FStanceModifiedEventData>(AsyncParams.TimeStep->BaseSimTimeMs, OldStance, ActiveStance, SetGeometryGTCallback));
+	Sim->AddEvent(MakeShared<FStanceModifiedEventData>(*AsyncParams.TimeStep, OldStance, ActiveStance, SetGeometryGTCallback));
 }
 
 void FChaosStanceModifier::OnStart_Async(const FMovementModifierParams_Async& AsyncParams)
@@ -193,7 +193,7 @@ void FChaosStanceModifier::OnStart_Async(const FMovementModifierParams_Async& As
 				CharMoverComp->Crouch();
 			}
 		};
-		Sim->AddEvent(MakeShared<FStanceModifierModeChangedEventData>(AsyncParams.TimeStep->BaseSimTimeMs, CurrentModeName, AsyncParams.SyncState->MovementMode, EnableWantsToCrouch));
+		Sim->AddEvent(MakeShared<FStanceModifierModeChangedEventData>(*AsyncParams.TimeStep, CurrentModeName, AsyncParams.SyncState->MovementMode, EnableWantsToCrouch));
 	}
 
 	// Set the current mode - this is used to check if the mode changes
@@ -214,7 +214,7 @@ void FChaosStanceModifier::OnEnd_Async(const FMovementModifierParams_Async& Asyn
 				CharMoverComp->UnCrouch();
 			}
 		};
-		Sim->AddEvent(MakeShared<FStanceModifierModeChangedEventData>(AsyncParams.TimeStep->BaseSimTimeMs, CurrentModeName, AsyncParams.SyncState->MovementMode, DisableWantsToCrouch));
+		Sim->AddEvent(MakeShared<FStanceModifierModeChangedEventData>(*AsyncParams.TimeStep, CurrentModeName, AsyncParams.SyncState->MovementMode, DisableWantsToCrouch));
 	}
 
 	CurrentModeName = NAME_None;
diff --git a/Engine/Plugins/Experimental/ChaosMover/Source/ChaosMover/Private/Character/Transitions/ChaosCharacterJumpCheck.cpp b/Engine/Plugins/Experimental/ChaosMover/Source/ChaosMover/Private/Character/Transitions/ChaosCharacterJumpCheck.cpp
--- a/Engine/Plugins/Experimental/ChaosMover/Source/ChaosMover/Private/Character/Transitions/ChaosCharacterJumpCheck.cpp
+++ b/Engine/Plugins/Experimental/ChaosMover/Source/ChaosMover/Private/Character/Transitions/ChaosCharacterJumpCheck.cpp
@@ -93,7 +93,7 @@ void UChaosCharacterJumpCheck::Trigger_Implementation(const FSimulationTickParam
 		{
 			StartingHeight = StartingSyncState->GetLocation_WorldSpace().Z;
 		}
-		Simulation->AddEvent(MakeShared<FJumpedEventData>(Params.TimeStep.BaseSimTimeMs, StartingHeight));
+		Simulation->AddEvent(MakeShared<FJumpedEventData>(Params.TimeStep, StartingHeight));
 
 		// Apply an equal and opposite impulse to the ground
 		if (FractionalGroundReactionImpulse > 0.0f)
diff --git a/Engine/Plugins/Experimental/ChaosMover/Source/ChaosMover/Private/PathedMovement/ChaosPathedMovementMode.cpp b/Engine/Plugins/Experimental/ChaosMover/Source/ChaosMover/Private/PathedMovement/ChaosPathedMovementMode.cpp
--- a/Engine/Plugins/Experimental/ChaosMover/Source/ChaosMover/Private/PathedMovement/ChaosPathedMovementMode.cpp
+++ b/Engine/Plugins/Experimental/ChaosMover/Source/ChaosMover/Private/PathedMovement/ChaosPathedMovementMode.cpp
@@ -249,32 +249,32 @@ void UChaosPathedMovementMode::SimulationTick_Implementation(const FSimulationTi
 
 	if (bIsInReverseJustChanged)
 	{
-		OnInReverseChanging(Params.TimeStep.BaseSimTimeMs, CurrentProperties.bWantsReversePlayback);
+		OnInReverseChanging(Params.TimeStep, CurrentProperties.bWantsReversePlayback);
 	}
 	if (bLoopingJustChanged)
 	{
-		OnLoopingChanging(Params.TimeStep.BaseSimTimeMs, CurrentProperties.bWantsLoopingPlayback);
+		OnLoopingChanging(Params.TimeStep, CurrentProperties.bWantsLoopingPlayback);
 	}
 	bool bOneWayJustChanged = CurrentProperties.bWantsOneWayPlayback != PreviousProperties.bWantsOneWayPlayback;
 	if (bOneWayJustChanged)
 	{
-		OnOneWayChanging(Params.TimeStep.BaseSimTimeMs, CurrentProperties.bWantsLoopingPlayback);
+		OnOneWayChanging(Params.TimeStep, CurrentProperties.bWantsLoopingPlayback);
 	}
 
 	if (bReceivedNewStartRequest)
 	{
-		OnPlaybackStarting(Params.TimeStep.BaseSimTimeMs);
+		OnPlaybackStarting(Params.TimeStep);
 		OutputPathedMovemenState.LatestMovementStartFrame = Params.TimeStep.ServerFrame;
 	}
 
 	if ((bWasStoppedJustRequested && !OutputPathedMovemenState.bHasFinished) || bHasJustFinished)
 	{
-		OnPlaybackStopping(Params.TimeStep.BaseSimTimeMs, bHasJustFinished);
+		OnPlaybackStopping(Params.TimeStep, bHasJustFinished);
 	}
 
 	if (bHasJustBounced)
 	{
-		OnPlaybackBounced(Params.TimeStep.BaseSimTimeMs);
+		OnPlaybackBounced(Params.TimeStep);
 	}
 
 	FMoverDefaultSyncState& OutputDefaultSyncState = OutputState.SyncState.SyncStateCollection.FindOrAddMutableDataByType<FMoverDefaultSyncState>();
@@ -313,12 +313,12 @@ bool UChaosPathedMovementMode::UsesMovementBasisTransform() const
 	return true;
 }
 
-void UChaosPathedMovementMode::OnPlaybackStarting(double BaseSimTimeMs)
+void UChaosPathedMovementMode::OnPlaybackStarting(const FMoverTimeStep& SimTime)
 {
 	// Add a pathed movement started event to the simulation (will be broadcast on GT during post sim)
 	if (Simulation)
 	{
-		Simulation->AddEvent(MakeShared<FPathedMovementStartedEventData>(BaseSimTimeMs));
+		Simulation->AddEvent(MakeShared<FPathedMovementStartedEventData>(SimTime));
 
 #if WITH_CHAOS_VISUAL_DEBUGGER
 	if (FChaosVisualDebuggerTrace::IsTracing())
@@ -329,12 +329,12 @@ void UChaosPathedMovementMode::OnPlaybackStarting(double BaseSimTimeMs)
 	}
 }
 
-void UChaosPathedMovementMode::OnPlaybackStopping(double BaseSimTimeMs, bool bReachedEndOfPlayback)
+void UChaosPathedMovementMode::OnPlaybackStopping(const FMoverTimeStep& SimTime, bool bReachedEndOfPlayback)
 {
 	// Add a pathed movement stopped event to the simulation (will be broadcast on GT during post sim)
 	if (Simulation)
 	{
-		Simulation->AddEvent(MakeShared<FPathedMovementStoppedEventData>(BaseSimTimeMs, bReachedEndOfPlayback));
+		Simulation->AddEvent(MakeShared<FPathedMovementStoppedEventData>(SimTime, bReachedEndOfPlayback));
 
 #if WITH_CHAOS_VISUAL_DEBUGGER
 	if (FChaosVisualDebuggerTrace::IsTracing())
@@ -346,12 +346,12 @@ void UChaosPathedMovementMode::OnPlaybackStopping(double BaseSimTimeMs, bool bRe
 	}
 }
 
-void UChaosPathedMovementMode::OnPlaybackBounced(double BaseSimTimeMs)
+void UChaosPathedMovementMode::OnPlaybackBounced(const FMoverTimeStep& SimTime)
 {
 	// Add a pathed movement bounced event to the simulation (will be broadcast on GT during post sim)
 	if (Simulation)
 	{
-		Simulation->AddEvent(MakeShared<FPathedMovementBouncedEventData>(BaseSimTimeMs));
+		Simulation->AddEvent(MakeShared<FPathedMovementBouncedEventData>(SimTime));
 
 #if WITH_CHAOS_VISUAL_DEBUGGER
 		if (FChaosVisualDebuggerTrace::IsTracing())
@@ -362,12 +362,12 @@ void UChaosPathedMovementMode::OnPlaybackBounced(double BaseSimTimeMs)
 	}
 }
 
-void UChaosPathedMovementMode::OnInReverseChanging(double BaseSimTimeMs, bool bNewInReverse)
+void UChaosPathedMovementMode::OnInReverseChanging(const FMoverTimeStep& SimTime, bool bNewInReverse)
 {
 	// Add a pathed movement stopped event to the simulation (will be broadcast on GT during post sim)
 	if (Simulation)
 	{
-		Simulation->AddEvent(MakeShared<FPathedMovementInReverseChangedEventData>(BaseSimTimeMs, bNewInReverse));
+		Simulation->AddEvent(MakeShared<FPathedMovementInReverseChangedEventData>(SimTime, bNewInReverse));
 
 #if WITH_CHAOS_VISUAL_DEBUGGER
 		if (FChaosVisualDebuggerTrace::IsTracing())
@@ -378,12 +378,12 @@ void UChaosPathedMovementMode::OnInReverseChanging(double BaseSimTimeMs, bool bN
 	}
 }
 
-void UChaosPathedMovementMode::OnLoopingChanging(double BaseSimTimeMs, bool bNewLooping)
+void UChaosPathedMovementMode::OnLoopingChanging(const FMoverTimeStep& SimTime, bool bNewLooping)
 {
 	// Add a pathed movement stopped event to the simulation (will be broadcast on GT during post sim)
 	if (Simulation)
 	{
-		Simulation->AddEvent(MakeShared<FPathedMovementLoopingChangedEventData>(BaseSimTimeMs, bNewLooping));
+		Simulation->AddEvent(MakeShared<FPathedMovementLoopingChangedEventData>(SimTime, bNewLooping));
 
 #if WITH_CHAOS_VISUAL_DEBUGGER
 		if (FChaosVisualDebuggerTrace::IsTracing())
@@ -394,12 +394,12 @@ void UChaosPathedMovementMode::OnLoopingChanging(double BaseSimTimeMs, bool bNew
 	}
 }
 
-void UChaosPathedMovementMode::OnOneWayChanging(double BaseSimTimeMs, bool bNewOneWay)
+void UChaosPathedMovementMode::OnOneWayChanging(const FMoverTimeStep& SimTime, bool bNewOneWay)
 {
 	// Add a pathed movement stopped event to the simulation (will be broadcast on GT during post sim)
 	if (Simulation)
 	{
-		Simulation->AddEvent(MakeShared<FPathedMovementOneWayChangedEventData>(BaseSimTimeMs, bNewOneWay));
+		Simulation->AddEvent(MakeShared<FPathedMovementOneWayChangedEventData>(SimTime, bNewOneWay));
 
 #if WITH_CHAOS_VISUAL_DEBUGGER
 		if (FChaosVisualDebuggerTrace::IsTracing())
diff --git a/Engine/Plugins/Experimental/ChaosMover/Source/ChaosMover/Public/ChaosMover/ChaosMoverSimulationTypes.h b/Engine/Plugins/Experimental/ChaosMover/Source/ChaosMover/Public/ChaosMover/ChaosMoverSimulationTypes.h
--- a/Engine/Plugins/Experimental/ChaosMover/Source/ChaosMover/Public/ChaosMover/ChaosMoverSimulationTypes.h
+++ b/Engine/Plugins/Experimental/ChaosMover/Source/ChaosMover/Public/ChaosMover/ChaosMoverSimulationTypes.h
@@ -278,8 +278,8 @@ struct FStanceModifiedEventData : public FMoverSimulationEventData
 {
 	GENERATED_BODY()
 
-	FStanceModifiedEventData(double InEventTimeMs, EStanceMode InOldStance, EStanceMode InNewStance, FEventProcessedCallbackPtr InEventProcessedCallback = nullptr)
-		: FMoverSimulationEventData(InEventTimeMs, InEventProcessedCallback)
+	FStanceModifiedEventData(const FMoverTimeStep& InEventTime, EStanceMode InOldStance, EStanceMode InNewStance, FEventProcessedCallbackPtr InEventProcessedCallback = nullptr)
+		: FMoverSimulationEventData(InEventTime, InEventProcessedCallback)
 		, OldStance(InOldStance)
 		, NewStance(InNewStance)
 	{
@@ -300,8 +300,8 @@ struct FStanceModifierModeChangedEventData : public FMoverSimulationEventData
 {
 	GENERATED_BODY()
 	
-	FStanceModifierModeChangedEventData(double InEventTimeMs, FName InOldMode, FName InNewMode, FEventProcessedCallbackPtr InEventProcessedCallback = nullptr)
-	: FMoverSimulationEventData(InEventTimeMs, InEventProcessedCallback)
+	FStanceModifierModeChangedEventData(const FMoverTimeStep& InEventTime, FName InOldMode, FName InNewMode, FEventProcessedCallbackPtr InEventProcessedCallback = nullptr)
+	: FMoverSimulationEventData(InEventTime, InEventProcessedCallback)
 	, OldMode(InOldMode)
 	, NewMode(InNewMode)
 	{
diff --git a/Engine/Plugins/Experimental/ChaosMover/Source/ChaosMover/Public/ChaosMover/ChaosMoverStateMachine.h b/Engine/Plugins/Experimental/ChaosMover/Source/ChaosMover/Public/ChaosMover/ChaosMoverStateMachine.h
--- a/Engine/Plugins/Experimental/ChaosMover/Source/ChaosMover/Public/ChaosMover/ChaosMoverStateMachine.h
+++ b/Engine/Plugins/Experimental/ChaosMover/Source/ChaosMover/Public/ChaosMover/ChaosMoverStateMachine.h
@@ -119,8 +119,7 @@ namespace UE::ChaosMover
 		bool ApplyInstantEffects(const FMoverTickStartData& SubstepStartData, const FMoverTimeStep& SubTimeStep, FMoverSyncState& OutputState);
 		void ReceiveInstantMovementEffects(const FMoverTimeStep& TimeStep, const FChaosNetInstantMovementEffectsQueue* InstantMovementEffectsQueue);
 	
-		double InternalSimTimeMs = 0.0f;
-		int32 InternalServerFrame = 0;
+		FMoverTimeStep InternalSimTime;
 
 		struct FIDHistory
 		{
diff --git a/Engine/Plugins/Experimental/ChaosMover/Source/ChaosMover/Public/ChaosMover/PathedMovement/ChaosPathedMovementMode.h b/Engine/Plugins/Experimental/ChaosMover/Source/ChaosMover/Public/ChaosMover/PathedMovement/ChaosPathedMovementMode.h
--- a/Engine/Plugins/Experimental/ChaosMover/Source/ChaosMover/Public/ChaosMover/PathedMovement/ChaosPathedMovementMode.h
+++ b/Engine/Plugins/Experimental/ChaosMover/Source/ChaosMover/Public/ChaosMover/PathedMovement/ChaosPathedMovementMode.h
@@ -66,12 +66,12 @@ public:
 	float GetOneWayTripDuration() const { return OneWayPlaybackDuration; }
 	
 protected:
-	void OnPlaybackStarting(double BaseSimTimeMs);
-	void OnPlaybackStopping(double BaseSimTimeMs, bool bReachedEndOfPlayback);
-	void OnPlaybackBounced(double BaseSimTimeMs);
-	void OnInReverseChanging(double BaseSimTimeMs, bool bNewInReverse);
-	void OnLoopingChanging(double BaseSimTimeMs, bool bNewLooping);
-	void OnOneWayChanging(double BaseSimTimeMs, bool bNewOneWay);
+	void OnPlaybackStarting(const FMoverTimeStep& SimTime);
+	void OnPlaybackStopping(const FMoverTimeStep& SimTime, bool bReachedEndOfPlayback);
+	void OnPlaybackBounced(const FMoverTimeStep& SimTime);
+	void OnInReverseChanging(const FMoverTimeStep& SimTime, bool bNewInReverse);
+	void OnLoopingChanging(const FMoverTimeStep& SimTime, bool bNewLooping);
+	void OnOneWayChanging(const FMoverTimeStep& SimTime, bool bNewOneWay);
 
 	void SetDefaultJointConstraintProperties();
 	
diff --git a/Engine/Plugins/Experimental/ChaosMover/Source/ChaosMover/Public/ChaosMover/PathedMovement/ChaosPathedMovementTypes.h b/Engine/Plugins/Experimental/ChaosMover/Source/ChaosMover/Public/ChaosMover/PathedMovement/ChaosPathedMovementTypes.h
--- a/Engine/Plugins/Experimental/ChaosMover/Source/ChaosMover/Public/ChaosMover/PathedMovement/ChaosPathedMovementTypes.h
+++ b/Engine/Plugins/Experimental/ChaosMover/Source/ChaosMover/Public/ChaosMover/PathedMovement/ChaosPathedMovementTypes.h
@@ -204,8 +204,8 @@ struct FPathedMovementStartedEventData : public FMoverSimulationEventData
 {
 	GENERATED_BODY()
 
-	FPathedMovementStartedEventData(double InEventTimeMs)
-		: FMoverSimulationEventData(InEventTimeMs)
+	FPathedMovementStartedEventData(const FMoverTimeStep& InEventTime)
+		: FMoverSimulationEventData(InEventTime)
 	{
 	}
 	FPathedMovementStartedEventData() {}
@@ -221,8 +221,8 @@ struct FPathedMovementStoppedEventData : public FMoverSimulationEventData
 {
 	GENERATED_BODY()
 
-	FPathedMovementStoppedEventData(double InEventTimeMs, bool InbReachedEndOfPlayback)
-		: FMoverSimulationEventData(InEventTimeMs)
+	FPathedMovementStoppedEventData(const FMoverTimeStep& InEventTime, bool InbReachedEndOfPlayback)
+		: FMoverSimulationEventData(InEventTime)
 		, bReachedEndOfPlayback(InbReachedEndOfPlayback)
 	{
 	}
@@ -241,8 +241,8 @@ struct FPathedMovementBouncedEventData : public FMoverSimulationEventData
 {
 	GENERATED_BODY()
 
-	FPathedMovementBouncedEventData(double InEventTimeMs)
-		: FMoverSimulationEventData(InEventTimeMs)
+	FPathedMovementBouncedEventData(const FMoverTimeStep& InEventTime)
+		: FMoverSimulationEventData(InEventTime)
 	{
 	}
 	FPathedMovementBouncedEventData() {}
@@ -258,8 +258,8 @@ struct FPathedMovementInReverseChangedEventData : public FMoverSimulationEventDa
 {
 	GENERATED_BODY()
 
-	FPathedMovementInReverseChangedEventData(double InEventTimeMs, bool InbNewInReverse)
-		: FMoverSimulationEventData(InEventTimeMs)
+	FPathedMovementInReverseChangedEventData(const FMoverTimeStep& InEventTime, bool InbNewInReverse)
+		: FMoverSimulationEventData(InEventTime)
 		, bNewInReverse(InbNewInReverse)
 	{
 	}
@@ -278,8 +278,8 @@ struct FPathedMovementLoopingChangedEventData : public FMoverSimulationEventData
 {
 	GENERATED_BODY()
 
-	FPathedMovementLoopingChangedEventData(double InEventTimeMs, bool InbNewLooping)
-		: FMoverSimulationEventData(InEventTimeMs)
+	FPathedMovementLoopingChangedEventData(const FMoverTimeStep& InEventTime, bool InbNewLooping)
+		: FMoverSimulationEventData(InEventTime)
 		, bNewLooping(InbNewLooping)
 	{
 	}
@@ -298,8 +298,8 @@ struct FPathedMovementOneWayChangedEventData : public FMoverSimulationEventData
 {
 	GENERATED_BODY()
 
-	FPathedMovementOneWayChangedEventData(double InEventTimeMs, bool InbNewOneWay)
-		: FMoverSimulationEventData(InEventTimeMs)
+	FPathedMovementOneWayChangedEventData(const FMoverTimeStep& InEventTime, bool InbNewOneWay)
+		: FMoverSimulationEventData(InEventTime)
 		, bNewOneWay(InbNewOneWay)
 	{
 	}
diff --git a/Engine/Plugins/Experimental/Mover/Source/Mover/Private/DefaultMovementSet/InstantMovementEffects/BasicInstantMovementEffects.cpp b/Engine/Plugins/Experimental/Mover/Source/Mover/Private/DefaultMovementSet/InstantMovementEffects/BasicInstantMovementEffects.cpp
--- a/Engine/Plugins/Experimental/Mover/Source/Mover/Private/DefaultMovementSet/InstantMovementEffects/BasicInstantMovementEffects.cpp
+++ b/Engine/Plugins/Experimental/Mover/Source/Mover/Private/DefaultMovementSet/InstantMovementEffects/BasicInstantMovementEffects.cpp
@@ -72,7 +72,7 @@ bool FTeleportEffect::ApplyMovementEffect(FApplyMovementEffectParams& ApplyEffec
 			SimBlackboard->Invalidate(CommonBlackboard::LastFoundDynamicMovementBase);
 		}
 
-		ApplyEffectParams.OutputEvents.Add(MakeShared<FTeleportSucceededEventData>(ApplyEffectParams.TimeStep->BaseSimTimeMs, PreviousLocation, PreviousRotation, TargetLocation, FQuat(FinalTargetRotation)));
+		ApplyEffectParams.OutputEvents.Add(MakeShared<FTeleportSucceededEventData>(*ApplyEffectParams.TimeStep, PreviousLocation, PreviousRotation, TargetLocation, FQuat(FinalTargetRotation)));
 
 		return true;
 	}
@@ -84,7 +84,7 @@ bool FTeleportEffect::ApplyMovementEffect(FApplyMovementEffectParams& ApplyEffec
 	}
 #endif // !defined(BUILD_SHIPPING) || !BUILD_SHIPPING
 
-	ApplyEffectParams.OutputEvents.Add(MakeShared<FTeleportFailedEventData>(ApplyEffectParams.TimeStep->BaseSimTimeMs, PreviousLocation, PreviousRotation, TargetLocation, FQuat(FinalTargetRotation), ETeleportFailureReason::Reason_NotAvailable));
+	ApplyEffectParams.OutputEvents.Add(MakeShared<FTeleportFailedEventData>(*ApplyEffectParams.TimeStep, PreviousLocation, PreviousRotation, TargetLocation, FQuat(FinalTargetRotation), ETeleportFailureReason::Reason_NotAvailable));
 
 	return false;
 }
diff --git a/Engine/Plugins/Experimental/Mover/Source/Mover/Private/MoverTypes.cpp b/Engine/Plugins/Experimental/Mover/Source/Mover/Private/MoverTypes.cpp
--- a/Engine/Plugins/Experimental/Mover/Source/Mover/Private/MoverTypes.cpp
+++ b/Engine/Plugins/Experimental/Mover/Source/Mover/Private/MoverTypes.cpp
@@ -36,6 +36,16 @@ FMoverOnImpactParams::FMoverOnImpactParams(const FName& ModeName, const FHitResu
 {
 }
 
+FMoverEventContext::FMoverEventContext()
+{
+}
+
+FMoverEventContext::FMoverEventContext(const FMoverTimeStep& InEventTime)
+	: EventTimeMs(InEventTime.BaseSimTimeMs)
+	, bIsDuringResimulation(InEventTime.bIsResimulating)
+{
+}
+
 FMoverDataStructBase::FMoverDataStructBase()
 {
 }
diff --git a/Engine/Plugins/Experimental/Mover/Source/Mover/Public/DefaultMovementSet/CharacterMoverSimulationTypes.h b/Engine/Plugins/Experimental/Mover/Source/Mover/Public/DefaultMovementSet/CharacterMoverSimulationTypes.h
--- a/Engine/Plugins/Experimental/Mover/Source/Mover/Public/DefaultMovementSet/CharacterMoverSimulationTypes.h
+++ b/Engine/Plugins/Experimental/Mover/Source/Mover/Public/DefaultMovementSet/CharacterMoverSimulationTypes.h
@@ -12,8 +12,8 @@ struct FLandedEventData : public FMoverSimulationEventData
 {
 	GENERATED_BODY()
 
-	FLandedEventData(double InEventTimeMs, const FHitResult& InHitResult, const FName InNewModeName)
-		: FMoverSimulationEventData(InEventTimeMs)
+	FLandedEventData(const FMoverTimeStep& InEventTime, const FHitResult& InHitResult, const FName InNewModeName)
+		: FMoverSimulationEventData(InEventTime)
 		, HitResult(InHitResult)
 		, NewModeName(InNewModeName)
 	{
@@ -34,8 +34,8 @@ struct FJumpedEventData : public FMoverSimulationEventData
 {
 	GENERATED_BODY()
 
-	FJumpedEventData(double InEventTimeMs, float InJumpStartHeight)
-		: FMoverSimulationEventData(InEventTimeMs)
+	FJumpedEventData(const FMoverTimeStep& InEventTime, float InJumpStartHeight)
+		: FMoverSimulationEventData(InEventTime)
 		, JumpStartHeight(InJumpStartHeight)
 	{
 	}
diff --git a/Engine/Plugins/Experimental/Mover/Source/Mover/Public/MoverSimulationTypes.h b/Engine/Plugins/Experimental/Mover/Source/Mover/Public/MoverSimulationTypes.h
--- a/Engine/Plugins/Experimental/Mover/Source/Mover/Public/MoverSimulationTypes.h
+++ b/Engine/Plugins/Experimental/Mover/Source/Mover/Public/MoverSimulationTypes.h
@@ -515,11 +515,12 @@ struct FMoverSimulationEventData
 
 	using FEventProcessedCallbackPtr = std::function<void(const FMoverSimulationEventData& Data, const FMoverSimEventGameThreadContext& GameThreadContext)>;
 
-	FMoverSimulationEventData(double InEventTimeMs, FEventProcessedCallbackPtr InEventProcessedCallback = nullptr)
+	FMoverSimulationEventData(const FMoverTimeStep& InEventTime, FEventProcessedCallbackPtr InEventProcessedCallback = nullptr)
 		: EventProcessedCallback(InEventProcessedCallback)
+		, Context(InEventTime)
 	{
-		Context.EventTimeMs = InEventTimeMs;
 	}
+
 	FMoverSimulationEventData() {}
 	virtual ~FMoverSimulationEventData() {}
 
@@ -565,8 +566,8 @@ struct FMovementModeChangedEventData : public FMoverSimulationEventData
 {
 	GENERATED_BODY()
 
-	FMovementModeChangedEventData(float InEventTimeMs, const FName InPreviousModeName, const FName InNewModeName, FEventProcessedCallbackPtr InEventProcessedCallback = nullptr)
-		: FMoverSimulationEventData(InEventTimeMs, InEventProcessedCallback)
+	FMovementModeChangedEventData(const FMoverTimeStep& InEventTime, const FName InPreviousModeName, const FName InNewModeName, FEventProcessedCallbackPtr InEventProcessedCallback = nullptr)
+		: FMoverSimulationEventData(InEventTime, InEventProcessedCallback)
 		, PreviousModeName(InPreviousModeName)
 		, NewModeName(InNewModeName)
 	{
@@ -587,8 +588,8 @@ struct FTeleportSucceededEventData : public FMoverSimulationEventData
 {
 	GENERATED_BODY()
 
-	FTeleportSucceededEventData(float InEventTimeMs, const FVector& InFromLocation, const FQuat& InFromRotation, const FVector& InToLocation, const FQuat& InToRotation)
-		: FMoverSimulationEventData(InEventTimeMs)
+	FTeleportSucceededEventData(const FMoverTimeStep& InEventTime, const FVector& InFromLocation, const FQuat& InFromRotation, const FVector& InToLocation, const FQuat& InToRotation)
+		: FMoverSimulationEventData(InEventTime)
 		, FromLocation(InFromLocation)
 		, FromRotation(InFromRotation)
 		, ToLocation(InToLocation)
@@ -619,8 +620,8 @@ struct FTeleportFailedEventData : public FMoverSimulationEventData
 {
 	GENERATED_BODY()
 
-	FTeleportFailedEventData(float InEventTimeMs, const FVector& InFromLocation, const FQuat& InFromRotation, const FVector& InToLocation, const FQuat& InToRotation, ETeleportFailureReason InTeleportFailureReason)
-		: FMoverSimulationEventData(InEventTimeMs)
+	FTeleportFailedEventData(const FMoverTimeStep& InEventTime, const FVector& InFromLocation, const FQuat& InFromRotation, const FVector& InToLocation, const FQuat& InToRotation, ETeleportFailureReason InTeleportFailureReason)
+		: FMoverSimulationEventData(InEventTime)
 		, FromLocation(InFromLocation)
 		, FromRotation(InFromRotation)
 		, ToLocation(InToLocation)
diff --git a/Engine/Plugins/Experimental/Mover/Source/Mover/Public/MoverTypes.h b/Engine/Plugins/Experimental/Mover/Source/Mover/Public/MoverTypes.h
--- a/Engine/Plugins/Experimental/Mover/Source/Mover/Public/MoverTypes.h
+++ b/Engine/Plugins/Experimental/Mover/Source/Mover/Public/MoverTypes.h
@@ -87,10 +87,18 @@ struct FMoverEventContext
 {
 	GENERATED_BODY()
 
+	UE_API FMoverEventContext();
+	UE_API FMoverEventContext(const FMoverTimeStep& InEventTime);
+
 	// Simulation time that this event occurred (in server simulation timespace)
 	UPROPERTY(EditAnywhere, BlueprintReadWrite, Category = Mover)
 	double EventTimeMs = -1.0;
 
+	// Did this event occur while resimulating a past sim frame?
+	UPROPERTY(EditAnywhere, BlueprintReadWrite, Category = Mover)
+	uint8 bIsDuringResimulation : 1 = 0;
+
+	// Did this event occur during the instant that the simulation rolled back to an earlier state? Note this is different than during resim. See also @bIsDuringResimulation
 	UPROPERTY(EditAnywhere, BlueprintReadWrite, Category = Mover)
 	uint8 bIsCausedByRollback : 1 = 0;
 };
