[Chaos Mover] Adding landing check time step context

diff --git a/Engine/Plugins/Experimental/ChaosMover/Source/ChaosMover/Private/Character/Transitions/ChaosCharacterLandingCheck.cpp b/Engine/Plugins/Experimental/ChaosMover/Source/ChaosMover/Private/Character/Transitions/ChaosCharacterLandingCheck.cpp
--- a/Engine/Plugins/Experimental/ChaosMover/Source/ChaosMover/Private/Character/Transitions/ChaosCharacterLandingCheck.cpp
+++ b/Engine/Plugins/Experimental/ChaosMover/Source/ChaosMover/Private/Character/Transitions/ChaosCharacterLandingCheck.cpp
@@ -112,7 +112,7 @@ TSharedPtr<FLandedEventData> UChaosCharacterLandingCheck::AllocLandedEventData()
 
 void UChaosCharacterLandingCheck::PopulateLandedEventData(const FSimulationTickParams& Params, const FFloorCheckResult& FloorResult, FLandedEventData& OutLandedEvent)
 {
-	OutLandedEvent.Context.EventTimeMs = Params.TimeStep.BaseSimTimeMs;
+	OutLandedEvent.Context = FMoverEventContext(Params.TimeStep);
 	OutLandedEvent.HitResult = FloorResult.HitResult;
 	OutLandedEvent.NewModeName = TransitionToGroundMode;
 }
