Removing use of TMap::RemoveCurrent in CullOldFrames and Rollback

diff --git a/Engine/Plugins/Experimental/ChaosMover/Source/ChaosMover/Private/ChaosMoverStateMachine.cpp b/Engine/Plugins/Experimental/ChaosMover/Source/ChaosMover/Private/ChaosMoverStateMachine.cpp
--- a/Engine/Plugins/Experimental/ChaosMover/Source/ChaosMover/Private/ChaosMoverStateMachine.cpp
+++ b/Engine/Plugins/Experimental/ChaosMover/Source/ChaosMover/Private/ChaosMoverStateMachine.cpp
@@ -900,24 +900,36 @@ namespace UE::ChaosMover
 	{
 		const int32 OldestFrameToKeep = CurrentFrame - MaxAge;
 
+		TArray<int32, TInlineAllocator<8>> KeysToRemove;
 		for (TMap<int32, TSet<uint8>>::TIterator It = IDsByFrame.CreateIterator(); It; ++It)
 		{
 			if (It.Key() < OldestFrameToKeep)
 			{
-				It.RemoveCurrent();
+				KeysToRemove.Add(It.Key());
 			}
 		}
+
+		for (int32 KeyToRemove : KeysToRemove)
+		{
+			IDsByFrame.Remove(KeyToRemove);
+		}
 	}
 	
 	void FMoverStateMachine::FIDHistory::Rollback(int32 RollbackToFrameInclusive)
 	{
+		TArray<int32, TInlineAllocator<8>> KeysToRemove;
 		for (TMap<int32, TSet<uint8>>::TIterator It = IDsByFrame.CreateIterator(); It; ++It)
 		{
 			if (It.Key() >= RollbackToFrameInclusive)
 			{
-				It.RemoveCurrent();
+				KeysToRemove.Add(It.Key());
 			}
 		}
+
+		for (int32 KeyToRemove : KeysToRemove)
+		{
+			IDsByFrame.Remove(KeyToRemove);
+		}
 	}
 
 	bool FMoverStateMachine::FIDHistory::WasIDAlreadySeen(uint8 ID) const
