ChaosMover: fixed bug in pathed movement mode where the OnOneWayChanging event was receiving the wrong flag

diff --git a/Engine/Plugins/Experimental/ChaosMover/Source/ChaosMover/Private/PathedMovement/ChaosPathedMovementMode.cpp b/Engine/Plugins/Experimental/ChaosMover/Source/ChaosMover/Private/PathedMovement/ChaosPathedMovementMode.cpp
--- a/Engine/Plugins/Experimental/ChaosMover/Source/ChaosMover/Private/PathedMovement/ChaosPathedMovementMode.cpp
+++ b/Engine/Plugins/Experimental/ChaosMover/Source/ChaosMover/Private/PathedMovement/ChaosPathedMovementMode.cpp
@@ -258,7 +258,7 @@ void UChaosPathedMovementMode::SimulationTick_Implementation(const FSimulationTi
 	bool bOneWayJustChanged = CurrentProperties.bWantsOneWayPlayback != PreviousProperties.bWantsOneWayPlayback;
 	if (bOneWayJustChanged)
 	{
-		OnOneWayChanging(Params.TimeStep, CurrentProperties.bWantsLoopingPlayback);
+		OnOneWayChanging(Params.TimeStep, CurrentProperties.bWantsOneWayPlayback);
 	}
 
 	if (bReceivedNewStartRequest)
