Fix Newton Movement SH

diff --git a/Engine/Plugins/Experimental/ChaosMover/Source/ChaosMover/Private/Character/Transitions/ChaosCharacterJumpCheck.cpp b/Engine/Plugins/Experimental/ChaosMover/Source/ChaosMover/Private/Character/Transitions/ChaosCharacterJumpCheck.cpp
--- a/Engine/Plugins/Experimental/ChaosMover/Source/ChaosMover/Private/Character/Transitions/ChaosCharacterJumpCheck.cpp
+++ b/Engine/Plugins/Experimental/ChaosMover/Source/ChaosMover/Private/Character/Transitions/ChaosCharacterJumpCheck.cpp
@@ -105,7 +105,10 @@ void UChaosCharacterJumpCheck::Trigger_Implementation(const FSimulationTickParam
 				FVector ImpulseToApply = -FractionalGroundReactionImpulse * CharacterMass * JumpVelocity;
 				if (Chaos::FPBDRigidParticleHandle* GroundParticle = Interface.GetRigidParticle(GroundObject))
 				{
-					Private::ApplyImpulse(GroundParticle, ImpulseToApply, FloorResult.HitResult.ImpactPoint);
+					if (GroundParticle->IsDynamic())
+					{
+						Private::ApplyImpulse(GroundParticle, ImpulseToApply, FloorResult.HitResult.ImpactPoint);
+					}
 				}
 			}
 		}
