[NEWTON] gravity is ignored for a moment when jumping
#tech-notes: Implement JumpForceTimeRemaining on Newton to ignore gravity when jumping

#changelist validated
#virtualized

diff --git a/Engine/Plugins/Experimental/ChaosMover/Source/ChaosMover/Private/ChaosMoverSimulation.cpp b/Engine/Plugins/Experimental/ChaosMover/Source/ChaosMover/Private/ChaosMoverSimulation.cpp
--- a/Engine/Plugins/Experimental/ChaosMover/Source/ChaosMover/Private/ChaosMoverSimulation.cpp
+++ b/Engine/Plugins/Experimental/ChaosMover/Source/ChaosMover/Private/ChaosMoverSimulation.cpp
@@ -897,6 +897,11 @@ void UChaosMoverSimulation::AddEvent(TSharedPtr<FMoverSimulationEventData> Event
 	}
 }
 
+const TArray<TSharedPtr<FMoverSimulationEventData>>& UChaosMoverSimulation::GetEvents() const
+{
+	return Events;
+}
+
 void UChaosMoverSimulation::ProcessSimulationEvent(const FMoverSimulationEventData& EventData)
 {
 	if (const FMovementModeChangedEventData* ModeChangedEventData = EventData.CastTo<FMovementModeChangedEventData>())
diff --git a/Engine/Plugins/Experimental/ChaosMover/Source/ChaosMover/Private/Character/Modes/ChaosFallingMode.cpp b/Engine/Plugins/Experimental/ChaosMover/Source/ChaosMover/Private/Character/Modes/ChaosFallingMode.cpp
--- a/Engine/Plugins/Experimental/ChaosMover/Source/ChaosMover/Private/Character/Modes/ChaosFallingMode.cpp
+++ b/Engine/Plugins/Experimental/ChaosMover/Source/ChaosMover/Private/Character/Modes/ChaosFallingMode.cpp
@@ -121,7 +121,7 @@ void UChaosFallingMode::GenerateMove_Implementation(const FMoverTickStartData& S
 	}
 
 	OutProposedMove = UAirMovementUtils::ComputeControlledFreeMove(Params);
-	const FVector VelocityWithGravity = StartVelocity + UMovementUtils::ComputeVelocityFromGravity(DefaultSimInputs->Gravity, DeltaSeconds);
+	const FVector VelocityWithGravity = bIgnoreGravity ? StartVelocity : StartVelocity + UMovementUtils::ComputeVelocityFromGravity(DefaultSimInputs->Gravity, DeltaSeconds);
 
 	//  If we are going faster than TerminalVerticalVelocity apply VerticalFallingDeceleration otherwise reset Z velocity to before we applied deceleration 
 	if (VelocityWithGravity.GetAbs().Dot(DefaultSimInputs->UpDir) > TerminalVerticalSpeed)
diff --git a/Engine/Plugins/Experimental/ChaosMover/Source/ChaosMover/Public/ChaosMover/ChaosMoverSimulation.h b/Engine/Plugins/Experimental/ChaosMover/Source/ChaosMover/Public/ChaosMover/ChaosMoverSimulation.h
--- a/Engine/Plugins/Experimental/ChaosMover/Source/ChaosMover/Public/ChaosMover/ChaosMoverSimulation.h
+++ b/Engine/Plugins/Experimental/ChaosMover/Source/ChaosMover/Public/ChaosMover/ChaosMoverSimulation.h
@@ -109,6 +109,7 @@ public:
 	CHAOSMOVER_API void ModifyContacts(const FMoverTimeStep& TimeStep, const UE::ChaosMover::FSimulationInputData& InputData, const UE::ChaosMover::FSimulationOutputData& OutputData, Chaos::FCollisionContactModifier& Modifier);
 
 	CHAOSMOVER_API void AddEvent(TSharedPtr<FMoverSimulationEventData> Event);
+	CHAOSMOVER_API const TArray<TSharedPtr<FMoverSimulationEventData>>& GetEvents() const;
 
 	CHAOSMOVER_API void ApplyNetInputData(const FMoverInputCmdContext& InNetInputCmd);
 	CHAOSMOVER_API void BuildNetInputData(FMoverInputCmdContext& OutNetInputCmd) const;
diff --git a/Engine/Plugins/Experimental/ChaosMover/Source/ChaosMover/Public/ChaosMover/Character/Modes/ChaosFallingMode.h b/Engine/Plugins/Experimental/ChaosMover/Source/ChaosMover/Public/ChaosMover/Character/Modes/ChaosFallingMode.h
--- a/Engine/Plugins/Experimental/ChaosMover/Source/ChaosMover/Public/ChaosMover/Character/Modes/ChaosFallingMode.h
+++ b/Engine/Plugins/Experimental/ChaosMover/Source/ChaosMover/Public/ChaosMover/Character/Modes/ChaosFallingMode.h
@@ -80,6 +80,13 @@ public:
 	 */
 	UPROPERTY(Category = Mover, EditAnywhere, BlueprintReadWrite, meta = (ClampMin = "0", ForceUnits = "cm/s"))
 	float TerminalVerticalSpeed;
+
+protected:
+
+	/*
+	* To allow other modes disable gravity
+	*/
+	uint8 bIgnoreGravity = 0;
 };
 
 #undef UE_API
