ChaosMover: Moving where we trace to CVD so we can record debug info in OnModifyContacts

diff --git a/Engine/Plugins/Experimental/ChaosMover/Source/ChaosMover/Private/Backends/ChaosMoverAsyncCallback.cpp b/Engine/Plugins/Experimental/ChaosMover/Source/ChaosMover/Private/Backends/ChaosMoverAsyncCallback.cpp
--- a/Engine/Plugins/Experimental/ChaosMover/Source/ChaosMover/Private/Backends/ChaosMoverAsyncCallback.cpp
+++ b/Engine/Plugins/Experimental/ChaosMover/Source/ChaosMover/Private/Backends/ChaosMoverAsyncCallback.cpp
@@ -117,6 +117,43 @@ namespace UE::ChaosMover
 		}
 	}
 
+	void FAsyncCallback::OnPostSolve_Internal()
+	{
+#if WITH_CHAOS_VISUAL_DEBUGGER
+		if (FChaosVisualDebuggerTrace::IsTracing())
+		{
+			QUICK_SCOPE_CYCLE_COUNTER(ChaosMover_AsyncCallback_OnPostSolve_Internal);
+
+			const FAsyncCallbackInput* AsyncInput = GetConsumerInput_Internal();
+
+			if (!AsyncInput || AsyncInput->InputData.IsEmpty())
+			{
+				return;
+			}
+
+			FMoverTimeStep TimeStep;
+			GetCurrentMoverTimeStep(AsyncInput, TimeStep);
+
+			// This array of resimulated backend indices is only used in the resim case
+			const TArray<int32, TInlineAllocator<64>>* ResimulatedBackendIndicesPtr =
+				(TimeStep.bIsResimulating && Chaos::ResimConsoleVars::bIsBubbleResimEnabled && !Chaos::ResimConsoleVars::bCanFrozenParticlesModifyContacts) ?
+				&ResimulatedBackendIndices : nullptr;
+
+			const int32 NumSimulatedBackends = ResimulatedBackendIndicesPtr ? ResimulatedBackendIndicesPtr->Num() : AsyncInput->Backends.Num();
+
+			const FAsyncCallbackOutput& AsyncOutput = GetProducerOutputData_Internal();
+			for (int Idx = 0; Idx < NumSimulatedBackends; ++Idx)
+			{
+				int32 BackendIndex = ResimulatedBackendIndicesPtr ? (*ResimulatedBackendIndicesPtr)[Idx] : Idx;
+				if (TStrongObjectPtr<UChaosMoverBackendComponent> Backend = AsyncInput->Backends[BackendIndex].Pin())
+				{
+					Backend->GetSimulation()->TraceMoverData(TimeStep, AsyncOutput.OutputData[BackendIndex]);
+				}
+			}
+		}
+#endif // WITH_CHAOS_VISUAL_DEBUGGER
+	}
+
 	void FAsyncCallback::GetCurrentMoverTimeStep(const FAsyncCallbackInput* AsyncInput, FMoverTimeStep& OutMoverTimeStep) const
 	{
 		OutMoverTimeStep.BaseSimTimeMs = GetSimTime_Internal() * 1000.0;
diff --git a/Engine/Plugins/Experimental/ChaosMover/Source/ChaosMover/Private/Backends/ChaosMoverAsyncCallback.h b/Engine/Plugins/Experimental/ChaosMover/Source/ChaosMover/Private/Backends/ChaosMoverAsyncCallback.h
--- a/Engine/Plugins/Experimental/ChaosMover/Source/ChaosMover/Private/Backends/ChaosMoverAsyncCallback.h
+++ b/Engine/Plugins/Experimental/ChaosMover/Source/ChaosMover/Private/Backends/ChaosMoverAsyncCallback.h
@@ -43,12 +43,14 @@ namespace UE::ChaosMover
 		FAsyncCallbackOutput,
 		Chaos::ESimCallbackOptions::Presimulate |
 		Chaos::ESimCallbackOptions::ContactModification |
+		Chaos::ESimCallbackOptions::PostSolve |
 		Chaos::ESimCallbackOptions::Rewind>
 	{
 	protected:
 		virtual void ProcessInputs_Internal(int32 PhysicsStep) override;
 		virtual void OnPreSimulate_Internal() override;
 		virtual void OnContactModification_Internal(Chaos::FCollisionContactModifier& Modifier) override;
+		virtual void OnPostSolve_Internal() override;
 
 		void GetCurrentMoverTimeStep(const FAsyncCallbackInput* AsyncInput, FMoverTimeStep& OutMoverTimeStep) const;
 
diff --git a/Engine/Plugins/Experimental/ChaosMover/Source/ChaosMover/Private/ChaosMoverSimulation.cpp b/Engine/Plugins/Experimental/ChaosMover/Source/ChaosMover/Private/ChaosMoverSimulation.cpp
--- a/Engine/Plugins/Experimental/ChaosMover/Source/ChaosMover/Private/ChaosMoverSimulation.cpp
+++ b/Engine/Plugins/Experimental/ChaosMover/Source/ChaosMover/Private/ChaosMoverSimulation.cpp
@@ -840,9 +840,6 @@ void UChaosMoverSimulation::OnPostSimulationTick(const FMoverTimeStep& TimeStep,
 	// Extract the events into the output data and clear
 	OutputData.Events = Events;
 	Events.Empty();
-
-	// Send Debug Data to the Chaos Visual Debugger
-	TraceMoverData(TimeStep, OutputData);
 }
 
 void UChaosMoverSimulation::TraceMoverData(const FMoverTimeStep& TimeStep, const UE::ChaosMover::FSimulationOutputData& OutputData)
diff --git a/Engine/Plugins/Experimental/ChaosMover/Source/ChaosMover/Public/ChaosMover/ChaosMoverSimulation.h b/Engine/Plugins/Experimental/ChaosMover/Source/ChaosMover/Public/ChaosMover/ChaosMoverSimulation.h
--- a/Engine/Plugins/Experimental/ChaosMover/Source/ChaosMover/Public/ChaosMover/ChaosMoverSimulation.h
+++ b/Engine/Plugins/Experimental/ChaosMover/Source/ChaosMover/Public/ChaosMover/ChaosMoverSimulation.h
@@ -160,6 +160,11 @@ protected:
 	// Actually perform the teleport
 	CHAOSMOVER_API virtual void Teleport(const FTransform& TargetTransform, const FMoverTimeStep& TimeStep, FMoverSyncState& OutputState);
 
+public:
+	// Trace mover data to Chaos Visual Debugger
+	CHAOSMOVER_API void TraceMoverData(const FMoverTimeStep& TimeStep, const UE::ChaosMover::FSimulationOutputData& OutputData);
+
+protected:
 	// Object State.
 	// In kinematic mode, the updated primitive component needs both its SimulatePhysics and UpdateKinematicFromSimulation options checked
 	CHAOSMOVER_API void SetControlledParticleDynamic();
@@ -179,8 +184,6 @@ protected:
 	CHAOSMOVER_API void SetActuationTargetTransform(const FTransform& TargetTransform);
 	CHAOSMOVER_API void TeleportActuationTarget(const FTransform& TargetTransform, bool AlsoTeleportControlledParticle = false);
 	
-	CHAOSMOVER_API void TraceMoverData(const FMoverTimeStep& TimeStep, const UE::ChaosMover::FSimulationOutputData& OutputData);
-
 	// State structs
 	FMoverSyncState CurrentSyncState;
 	// Data internal to the simulation
diff --git a/Engine/Plugins/Experimental/ChaosMover/Source/ChaosMover/Public/ChaosMover/Character/ChaosCharacterInputs.h b/Engine/Plugins/Experimental/ChaosMover/Source/ChaosMover/Public/ChaosMover/Character/ChaosCharacterInputs.h
--- a/Engine/Plugins/Experimental/ChaosMover/Source/ChaosMover/Public/ChaosMover/Character/ChaosCharacterInputs.h
+++ b/Engine/Plugins/Experimental/ChaosMover/Source/ChaosMover/Public/ChaosMover/Character/ChaosCharacterInputs.h
@@ -54,7 +54,7 @@ struct FChaosMoverCrouchInputs : public FMoverDataStructBase
 
 public:
 	// Launch velocity in cm/s or impulse in kg cm/s
-	UPROPERTY(BlueprintReadWrite, Category = "Chaos Mover")
+	UPROPERTY(VisibleAnywhere, BlueprintReadWrite, Category = "Chaos Mover")
 	bool bWantsToCrouch = false;
 
 	virtual UScriptStruct* GetScriptStruct() const override
