ChaosMover: TeleportTo instant movement effect for debugging

diff --git a/Engine/Plugins/Experimental/ChaosMover/Source/ChaosMover/Private/ChaosMoverConsoleVariables.cpp b/Engine/Plugins/Experimental/ChaosMover/Source/ChaosMover/Private/ChaosMoverConsoleVariables.cpp
--- a/Engine/Plugins/Experimental/ChaosMover/Source/ChaosMover/Private/ChaosMoverConsoleVariables.cpp
+++ b/Engine/Plugins/Experimental/ChaosMover/Source/ChaosMover/Private/ChaosMoverConsoleVariables.cpp
@@ -69,10 +69,51 @@ namespace UE::ChaosMover
 					};
 
 					TSharedPtr<FAsyncLocalOnlyInstantMovementEffect> EffectPtr = MakeShared<FAsyncLocalOnlyInstantMovementEffect>();
-					static FName EffectName("TeleportOnlyInstantMovementEffect");
+					static FName EffectName("TeleportLocalOnlyInstantMovementEffect");
 					EffectPtr->OptionalName = EffectName;
 					EffectPtr->AsyncFunction = TeleportLambda;
 					MoverComponent->QueueInstantMovementEffect(EffectPtr);
 				}));
 	}
+
+	static FAutoConsoleCommandWithWorldAndArgs TeleportToCommand
+	(
+		TEXT("ChaosMover.Debug.TeleportTo"),
+		TEXT("Usage = ChaosMover.Debug.TeleportTo <X> <Y> <Z> <Yaw> <Pitch> <Roll>. Teleports the ChaosMover character to (X, Y, Z) (Yaw, Pitch, Roll). Use '=' to keep the component unchanged, for instance 'ChaosMover.Debug.TeleportTo 0 0 = 0' will teleport to X=0,Y=0 leaving Z unchanged and set Yaw to 0, leaving pitch and roll unchanged"),
+		FConsoleCommandWithWorldAndArgsDelegate::CreateStatic([](const TArray<FString>& Args, UWorld* World)
+			{
+				APlayerController* LocalPlayerController = World ? World->GetFirstPlayerController() : nullptr;
+				APawn* LocalPawn = LocalPlayerController ? LocalPlayerController->GetPawn() : nullptr;
+				UMoverComponent* MoverComponent = LocalPawn ? LocalPawn->FindComponentByClass<UMoverComponent>() : nullptr;
+				if (MoverComponent)
+				{
+					TSharedPtr<FDebugTeleportToInstantMovementEffect> EffectPtr = MakeShared<FDebugTeleportToInstantMovementEffect>();
+					// Leaving a component unchanged with '=' means keep that character's location or rotation component unchanged
+					if (Args.Num() > 0 && Args[0].Len() > 0 && Args[0][0] != '=')
+					{
+						LexFromString(EffectPtr->TeleportLocation.X, *Args[0]);
+					}
+					if (Args.Num() > 1 && Args[1].Len() > 0 && Args[1][0] != '=')
+					{
+						LexFromString(EffectPtr->TeleportLocation.Y, *Args[1]);
+					}
+					if (Args.Num() > 2 && Args[2].Len() > 0 && Args[2][0] != '=')
+					{
+						LexFromString(EffectPtr->TeleportLocation.Z, *Args[2]);
+					}
+					if (Args.Num() > 3 && Args[3].Len() > 0 && Args[3][0] != '=')
+					{
+						LexFromString(EffectPtr->TeleportRotation.Yaw, *Args[3]);
+					}
+					if (Args.Num() > 4 && Args[4].Len() > 0 && Args[4][0] != '=')
+					{
+						LexFromString(EffectPtr->TeleportRotation.Pitch, *Args[4]);
+					}
+					if (Args.Num() > 5 && Args[5].Len() > 0 && Args[5][0] != '=')
+					{
+						LexFromString(EffectPtr->TeleportRotation.Roll, *Args[5]);
+					}
+					MoverComponent->ScheduleInstantMovementEffect(EffectPtr);
+				}
+			}));
 }
\ No newline at end of file
diff --git a/Engine/Plugins/Experimental/ChaosMover/Source/ChaosMover/Private/ChaosMoverSimulationTypes.cpp b/Engine/Plugins/Experimental/ChaosMover/Source/ChaosMover/Private/ChaosMoverSimulationTypes.cpp
--- a/Engine/Plugins/Experimental/ChaosMover/Source/ChaosMover/Private/ChaosMoverSimulationTypes.cpp
+++ b/Engine/Plugins/Experimental/ChaosMover/Source/ChaosMover/Private/ChaosMoverSimulationTypes.cpp
@@ -265,7 +265,8 @@ void FChaosNetInstantMovementEffectsQueue::Decay(float DecayAmount)
 
 bool FAsyncLocalOnlyInstantMovementEffect::ApplyMovementEffect_Async(FApplyMovementEffectParams_Async& ApplyEffectParams, FMoverSyncState& OutputState)
 {
-	if (bExecuteLocally && AsyncFunction.IsSet())
+	// AsyncFunction will not be set if this is not executing on the instance that queued this effect, nothing will execute
+	if (AsyncFunction.IsSet())
 	{
 		return AsyncFunction(ApplyEffectParams, OutputState);
 	}
@@ -279,11 +280,7 @@ FInstantMovementEffect* FAsyncLocalOnlyInstantMovementEffect::Clone() const
 
 void FAsyncLocalOnlyInstantMovementEffect::NetSerialize(FArchive& Ar)
 {
-	// Serialize nothing out, and when loading (receiving from the network), set to not execute locally
-	if (Ar.IsLoading())
-	{
-		bExecuteLocally = false;
-	}
+	// No serialization, this is a local only effect
 }
 
 UScriptStruct* FAsyncLocalOnlyInstantMovementEffect::GetScriptStruct() const
@@ -295,3 +292,45 @@ FString FAsyncLocalOnlyInstantMovementEffect::ToSimpleString() const
 {
 	return FString::Printf(TEXT("FAsyncLocalOnlyInstantMovementEffect %s"), *OptionalName.ToString());
 }
+
+bool FDebugTeleportToInstantMovementEffect::ApplyMovementEffect_Async(FApplyMovementEffectParams_Async& ApplyEffectParams, FMoverSyncState& OutputState)
+{
+	if (UChaosMoverSimulation* ChaosMoverSimulation = Cast<UChaosMoverSimulation>(ApplyEffectParams.Simulation))
+	{
+		if (const FMoverDefaultSyncState* DefaultSyncState = OutputState.SyncStateCollection.FindDataByType<FMoverDefaultSyncState>())
+		{
+			FVector FinalTeleportLocation = DefaultSyncState->GetLocation_WorldSpace();
+			FinalTeleportLocation.X = (TeleportLocation.X != MAX_FLT) ? TeleportLocation.X : FinalTeleportLocation.X;
+			FinalTeleportLocation.Y = (TeleportLocation.Y != MAX_FLT) ? TeleportLocation.Y : FinalTeleportLocation.Y;
+			FinalTeleportLocation.Z = (TeleportLocation.Z != MAX_FLT) ? TeleportLocation.Z : FinalTeleportLocation.Z;
+			FRotator FinalTeleportRotation = DefaultSyncState->GetOrientation_WorldSpace();
+			FinalTeleportRotation.Yaw = (TeleportRotation.Yaw != MAX_FLT) ? TeleportRotation.Yaw : FinalTeleportRotation.Yaw;
+			FinalTeleportRotation.Pitch = (TeleportRotation.Pitch != MAX_FLT) ? TeleportRotation.Pitch : FinalTeleportRotation.Pitch;
+			FinalTeleportRotation.Roll = (TeleportRotation.Roll != MAX_FLT) ? TeleportRotation.Roll : FinalTeleportRotation.Roll;
+			ChaosMoverSimulation->AttemptTeleport(*ApplyEffectParams.TimeStep, FTransform(FinalTeleportRotation, FinalTeleportLocation), /* bUseActorRotation = */ false, OutputState);
+			return true;
+		}
+	}
+	return false;
+}
+
+FInstantMovementEffect* FDebugTeleportToInstantMovementEffect::Clone() const
+{
+	return new FDebugTeleportToInstantMovementEffect(*this);
+}
+
+void FDebugTeleportToInstantMovementEffect::NetSerialize(FArchive& Ar)
+{
+	Ar << TeleportLocation;
+	Ar << TeleportRotation;
+}
+
+UScriptStruct* FDebugTeleportToInstantMovementEffect::GetScriptStruct() const
+{
+	return FDebugTeleportToInstantMovementEffect::StaticStruct();
+}
+
+FString FDebugTeleportToInstantMovementEffect::ToSimpleString() const
+{
+	return FString::Printf(TEXT("FDebugTeleportToInstantMovementEffect. TeleportLocation: %s\nTeleportRotation: %s"), *TeleportLocation.ToString(), *TeleportRotation.ToString());
+}
diff --git a/Engine/Plugins/Experimental/ChaosMover/Source/ChaosMover/Public/ChaosMover/ChaosMoverSimulationTypes.h b/Engine/Plugins/Experimental/ChaosMover/Source/ChaosMover/Public/ChaosMover/ChaosMoverSimulationTypes.h
--- a/Engine/Plugins/Experimental/ChaosMover/Source/ChaosMover/Public/ChaosMover/ChaosMoverSimulationTypes.h
+++ b/Engine/Plugins/Experimental/ChaosMover/Source/ChaosMover/Public/ChaosMover/ChaosMoverSimulationTypes.h
@@ -394,8 +394,9 @@ struct FChaosNetInstantMovementEffectsQueue : public FMoverDataStructBase
 	virtual void Decay(float DecayAmount) override;
 };
 
-// This instant movement effect is ChaosMover only and applies only on the instance issuing it
-// It will execute AsyncFunction on the physics thread if set
+
+// This instant movement effect is ChaosMover only and executes AsyncFunction on the issuing end point only
+// AsyncFunction does not (cannot) be serialized, so this is local only
 USTRUCT()
 struct FAsyncLocalOnlyInstantMovementEffect : public FInstantMovementEffect
 {
@@ -407,7 +408,26 @@ struct FAsyncLocalOnlyInstantMovementEffect : public FInstantMovementEffect
 	FName OptionalName;
 	// Async Function executed in ApplyMovementEffect_Async (only locally)
 	TFunction<bool(FApplyMovementEffectParams_Async&, FMoverSyncState&)> AsyncFunction;
-	bool bExecuteLocally = true;
+
+	UE_API virtual bool ApplyMovementEffect_Async(FApplyMovementEffectParams_Async& ApplyEffectParams, FMoverSyncState& OutputState) override;
+	UE_API virtual FInstantMovementEffect* Clone() const override;
+	UE_API virtual void NetSerialize(FArchive& Ar) override;
+	UE_API virtual UScriptStruct* GetScriptStruct() const override;
+	UE_API virtual FString ToSimpleString() const override;
+};
+
+// This instant movement effect is ChaosMover only and teleport to a location without any checks whatsoever, this is meant for debugging only
+USTRUCT()
+struct FDebugTeleportToInstantMovementEffect : public FInstantMovementEffect
+{
+	GENERATED_BODY()
+
+	UE_API FDebugTeleportToInstantMovementEffect() {}
+	virtual ~FDebugTeleportToInstantMovementEffect() {}
+
+	// A component can be set to MAX_FLT to signify "keep that current component unchanged"
+	FVector TeleportLocation = FVector(MAX_FLT, MAX_FLT, MAX_FLT);
+	FRotator TeleportRotation = FRotator(MAX_FLT, MAX_FLT, MAX_FLT);
 
 	UE_API virtual bool ApplyMovementEffect_Async(FApplyMovementEffectParams_Async& ApplyEffectParams, FMoverSyncState& OutputState) override;
 	UE_API virtual FInstantMovementEffect* Clone() const override;
