Build fix for Issue 1123456: PVS Studio warnings in ChaosMoverConsoleVariables.cpp

ChaosMoverConsoleVariables.cpp(86): warning V623: Consider inspecting the '?:' operator. A temporary object of the 'TObjectPtr' type is being created and subsequently destroyed. Check third operand.

diff --git a/Engine/Plugins/Experimental/ChaosMover/Source/ChaosMover/Private/ChaosMoverConsoleVariables.cpp b/Engine/Plugins/Experimental/ChaosMover/Source/ChaosMover/Private/ChaosMoverConsoleVariables.cpp
--- a/Engine/Plugins/Experimental/ChaosMover/Source/ChaosMover/Private/ChaosMoverConsoleVariables.cpp
+++ b/Engine/Plugins/Experimental/ChaosMover/Source/ChaosMover/Private/ChaosMoverConsoleVariables.cpp
@@ -82,38 +82,40 @@ namespace UE::ChaosMover
 		TEXT("Usage = ChaosMover.Debug.TeleportTo <X> <Y> <Z> <Yaw> <Pitch> <Roll>. Teleports the ChaosMover character to (X, Y, Z) (Yaw, Pitch, Roll). Use '=' to keep the component unchanged, for instance 'ChaosMover.Debug.TeleportTo 0 0 = 0' will teleport to X=0,Y=0 leaving Z unchanged and set Yaw to 0, leaving pitch and roll unchanged"),
 		FConsoleCommandWithWorldAndArgsDelegate::CreateStatic([](const TArray<FString>& Args, UWorld* World)
 			{
-				APlayerController* LocalPlayerController = World ? World->GetFirstPlayerController() : nullptr;
-				APawn* LocalPawn = LocalPlayerController ? LocalPlayerController->GetPawn() : nullptr;
-				UMoverComponent* MoverComponent = LocalPawn ? LocalPawn->FindComponentByClass<UMoverComponent>() : nullptr;
-				if (MoverComponent)
+				if (APlayerController* LocalPlayerController = World ? World->GetFirstPlayerController() : nullptr)
 				{
-					TSharedPtr<FDebugTeleportToInstantMovementEffect> EffectPtr = MakeShared<FDebugTeleportToInstantMovementEffect>();
-					// Leaving a component unchanged with '=' means keep that character's location or rotation component unchanged
-					if (Args.Num() > 0 && Args[0].Len() > 0 && Args[0][0] != '=')
-					{
-						LexFromString(EffectPtr->TeleportLocation.X, *Args[0]);
-					}
-					if (Args.Num() > 1 && Args[1].Len() > 0 && Args[1][0] != '=')
-					{
-						LexFromString(EffectPtr->TeleportLocation.Y, *Args[1]);
-					}
-					if (Args.Num() > 2 && Args[2].Len() > 0 && Args[2][0] != '=')
-					{
-						LexFromString(EffectPtr->TeleportLocation.Z, *Args[2]);
-					}
-					if (Args.Num() > 3 && Args[3].Len() > 0 && Args[3][0] != '=')
-					{
-						LexFromString(EffectPtr->TeleportRotation.Yaw, *Args[3]);
-					}
-					if (Args.Num() > 4 && Args[4].Len() > 0 && Args[4][0] != '=')
-					{
-						LexFromString(EffectPtr->TeleportRotation.Pitch, *Args[4]);
-					}
-					if (Args.Num() > 5 && Args[5].Len() > 0 && Args[5][0] != '=')
+					TObjectPtr<APawn> LocalPawn = LocalPlayerController->GetPawn();
+					UMoverComponent* MoverComponent = LocalPawn ? LocalPawn->FindComponentByClass<UMoverComponent>() : nullptr;
+					if (MoverComponent)
 					{
-						LexFromString(EffectPtr->TeleportRotation.Roll, *Args[5]);
+						TSharedPtr<FDebugTeleportToInstantMovementEffect> EffectPtr = MakeShared<FDebugTeleportToInstantMovementEffect>();
+						// Leaving a component unchanged with '=' means keep that character's location or rotation component unchanged
+						if (Args.Num() > 0 && Args[0].Len() > 0 && Args[0][0] != '=')
+						{
+							LexFromString(EffectPtr->TeleportLocation.X, *Args[0]);
+						}
+						if (Args.Num() > 1 && Args[1].Len() > 0 && Args[1][0] != '=')
+						{
+							LexFromString(EffectPtr->TeleportLocation.Y, *Args[1]);
+						}
+						if (Args.Num() > 2 && Args[2].Len() > 0 && Args[2][0] != '=')
+						{
+							LexFromString(EffectPtr->TeleportLocation.Z, *Args[2]);
+						}
+						if (Args.Num() > 3 && Args[3].Len() > 0 && Args[3][0] != '=')
+						{
+							LexFromString(EffectPtr->TeleportRotation.Yaw, *Args[3]);
+						}
+						if (Args.Num() > 4 && Args[4].Len() > 0 && Args[4][0] != '=')
+						{
+							LexFromString(EffectPtr->TeleportRotation.Pitch, *Args[4]);
+						}
+						if (Args.Num() > 5 && Args[5].Len() > 0 && Args[5][0] != '=')
+						{
+							LexFromString(EffectPtr->TeleportRotation.Roll, *Args[5]);
+						}
+						MoverComponent->ScheduleInstantMovementEffect(EffectPtr);
 					}
-					MoverComponent->ScheduleInstantMovementEffect(EffectPtr);
 				}
 			}));
 }
\ No newline at end of file
