[Chaos Mover] Check events have not already been fired before firing them again after a resimulation

diff --git a/Engine/Plugins/Experimental/ChaosMover/Source/ChaosMover/Private/Backends/ChaosMoverBackend.cpp b/Engine/Plugins/Experimental/ChaosMover/Source/ChaosMover/Private/Backends/ChaosMoverBackend.cpp
--- a/Engine/Plugins/Experimental/ChaosMover/Source/ChaosMover/Private/Backends/ChaosMoverBackend.cpp
+++ b/Engine/Plugins/Experimental/ChaosMover/Source/ChaosMover/Private/Backends/ChaosMoverBackend.cpp
@@ -525,6 +525,7 @@ void UChaosMoverBackendComponent::ProduceInputData(int32 PhysicsStep, int32 NumS
 	SimInputs.bIsRemotelyControlled = OwnerAsPawn ? (OwnerAsPawn->GetController() && !OwnerAsPawn->IsLocallyControlled()) : false;
 	SimInputs.OwningActor = GetOwner();
 	SimInputs.World = GetWorld();
+	SimInputs.AsyncNetworkPhysicsComponent = NetworkPhysicsComponent ? NetworkPhysicsComponent->GetNetworkPhysicsComponent_Internal() : nullptr;
 
 	if (UPrimitiveComponent* PrimComp = Cast<UPrimitiveComponent>(MoverComp.GetUpdatedComponent()))
 	{
diff --git a/Engine/Plugins/Experimental/ChaosMover/Source/ChaosMover/Private/ChaosMoverConsoleVariables.cpp b/Engine/Plugins/Experimental/ChaosMover/Source/ChaosMover/Private/ChaosMoverConsoleVariables.cpp
--- a/Engine/Plugins/Experimental/ChaosMover/Source/ChaosMover/Private/ChaosMoverConsoleVariables.cpp
+++ b/Engine/Plugins/Experimental/ChaosMover/Source/ChaosMover/Private/ChaosMoverConsoleVariables.cpp
@@ -39,6 +39,11 @@ namespace UE::ChaosMover
 		FAutoConsoleVariableRef CVarChaosMoverInstantMovementEffectIDHistorySize(TEXT("ChaosMover.Networking.InstantMovementEffectIDHistorySize"),
 			InstantMovementEffectIDHistorySize, TEXT("Number of past frames for which we keep the last seen instant movement effect IDs, to avoid duplicate processing of events when input is repeated\n"));
 
+		int32 FrameDifferenceLeniencyForEventTimeComparison = 1;
+		FAutoConsoleVariableRef CVarChaosMoverFrameDifferenceLeniencyForEventTimeComparison(TEXT("ChaosMover.Networking.FrameDifferenceLeniencyForEventTimeComparison"),
+			FrameDifferenceLeniencyForEventTimeComparison, TEXT("Within this number of frames events are considered to be at the same time and will not trigger multiple times.\n"));
+
+
 		static FAutoConsoleCommandWithWorld TeleportOnlyClientCommand
 		(
 			TEXT("ChaosMover.Debug.TeleportClientOnly"),
diff --git a/Engine/Plugins/Experimental/ChaosMover/Source/ChaosMover/Private/ChaosMoverSimulation.cpp b/Engine/Plugins/Experimental/ChaosMover/Source/ChaosMover/Private/ChaosMoverSimulation.cpp
--- a/Engine/Plugins/Experimental/ChaosMover/Source/ChaosMover/Private/ChaosMoverSimulation.cpp
+++ b/Engine/Plugins/Experimental/ChaosMover/Source/ChaosMover/Private/ChaosMoverSimulation.cpp
@@ -31,6 +31,7 @@
 #include "MoverSimulationTypes.h"
 #include "PBDRigidsSolver.h"
 #include "Physics/Experimental/PhysScene_Chaos.h"
+#include "Physics/NetworkPhysicsComponent.h"
 #include "PhysicsEngine/PhysicsObjectExternalInterface.h"
 #include "PhysicsProxy/CharacterGroundConstraintProxy.h"
 #include "UObject/UObjectGlobals.h"
@@ -837,7 +838,60 @@ void UChaosMoverSimulation::OnPostSimulationTick(const FMoverTimeStep& TimeStep,
 
 	CurrentSyncState = OutputData.SyncState;
 
-	// Extract the events into the output data and clear
+	// Extract the events into the output data
+	UpdateEvents(TimeStep, OutputData);
+}
+
+void UChaosMoverSimulation::UpdateEvents(const FMoverTimeStep& TimeStep, UE::ChaosMover::FSimulationOutputData& OutputData)
+{
+	if (const FChaosMoverSimulationDefaultInputs* SimInputs = LocalSimInput.FindDataByType<FChaosMoverSimulationDefaultInputs>())
+	{
+		// If networking is enabled then we can resimulate and potentially fire the same event
+		// multiple times. To avoid this we maintain a list of processed events and if the same
+		// event is generated at the same time then remove it
+		if (SimInputs->AsyncNetworkPhysicsComponent && !SimInputs->AsyncNetworkPhysicsComponent->IsServer())
+		{
+			// Clear events that are outside of the resimulation window
+			const int32 LastReceivedServerFrame = SimInputs->AsyncNetworkPhysicsComponent->GetLastReceivedServerFrame();
+			if (LastReceivedServerFrame > 0)
+			{
+				ProcessedEvents.RemoveAllSwap([&LastReceivedServerFrame](const TSharedPtr<FMoverSimulationEventData> Event) {
+					return Event->Context.ServerFrame <= LastReceivedServerFrame;
+					});
+			}
+
+			if (TimeStep.bIsResimulating)
+			{
+				// Remove events that have already been processed in the main update
+
+				Events.RemoveAll([this](const TSharedPtr<FMoverSimulationEventData> Event) {
+					const int32 AllowedDiff = UE::ChaosMover::CVars::FrameDifferenceLeniencyForEventTimeComparison;
+					for (TSharedPtr<FMoverSimulationEventData>& ProcessedEvent : ProcessedEvents)
+					{
+						if (const FMoverSimulationEventData* ProcessedEventPtr = ProcessedEvent.Get())
+						{
+							if (Event->IsEqual(*ProcessedEventPtr, AllowedDiff))
+							{
+								return true;
+							}
+						}
+						
+					}
+					return false;
+					});
+			}
+			else
+			{
+				// Add new events to the processed list
+				for (TSharedPtr<FMoverSimulationEventData>& Event : Events)
+				{
+					ProcessedEvents.Add(Event);
+				}
+			}
+		}
+	}
+
+	// Extract events into output data
 	OutputData.Events = Events;
 	Events.Empty();
 }
diff --git a/Engine/Plugins/Experimental/ChaosMover/Source/ChaosMover/Private/ChaosMoverSimulationTypes.cpp b/Engine/Plugins/Experimental/ChaosMover/Source/ChaosMover/Private/ChaosMoverSimulationTypes.cpp
--- a/Engine/Plugins/Experimental/ChaosMover/Source/ChaosMover/Private/ChaosMoverSimulationTypes.cpp
+++ b/Engine/Plugins/Experimental/ChaosMover/Source/ChaosMover/Private/ChaosMoverSimulationTypes.cpp
@@ -20,6 +20,7 @@ void FChaosMoverSimulationDefaultInputs::Reset()
 	PhysicsObject = nullptr;
 	OwningActor = nullptr;
 	World = nullptr;
+	AsyncNetworkPhysicsComponent = nullptr;
 	CollisionChannel = ECC_Pawn;
 }
 
diff --git a/Engine/Plugins/Experimental/ChaosMover/Source/ChaosMover/Public/ChaosMover/ChaosMoverConsoleVariables.h b/Engine/Plugins/Experimental/ChaosMover/Source/ChaosMover/Public/ChaosMover/ChaosMoverConsoleVariables.h
--- a/Engine/Plugins/Experimental/ChaosMover/Source/ChaosMover/Public/ChaosMover/ChaosMoverConsoleVariables.h
+++ b/Engine/Plugins/Experimental/ChaosMover/Source/ChaosMover/Public/ChaosMover/ChaosMoverConsoleVariables.h
@@ -14,5 +14,6 @@ namespace UE::ChaosMover
 		CHAOSMOVER_API extern bool bDrawOverlapQueries;
 		CHAOSMOVER_API extern bool bSkipGenerateMoveIfOverridden;
 		CHAOSMOVER_API extern int32 InstantMovementEffectIDHistorySize;
+		CHAOSMOVER_API extern int32 FrameDifferenceLeniencyForEventTimeComparison;
 	}
 };
diff --git a/Engine/Plugins/Experimental/ChaosMover/Source/ChaosMover/Public/ChaosMover/ChaosMoverSimulation.h b/Engine/Plugins/Experimental/ChaosMover/Source/ChaosMover/Public/ChaosMover/ChaosMoverSimulation.h
--- a/Engine/Plugins/Experimental/ChaosMover/Source/ChaosMover/Public/ChaosMover/ChaosMoverSimulation.h
+++ b/Engine/Plugins/Experimental/ChaosMover/Source/ChaosMover/Public/ChaosMover/ChaosMoverSimulation.h
@@ -57,7 +57,6 @@ public:
 	// Queue a scheduled instant movement effect
 	CHAOSMOVER_API void QueueInstantMovementEffect(const FScheduledInstantMovementEffect& ScheduledInstantMovementEffect, bool bShouldRollBack = true);
 
-public:
 	UFUNCTION(BlueprintCallable, CustomThunk, Category = Mover, meta = (CustomStructureParam = "MoveAsRawData", AllowAbstract = "false", DisplayName = "Queue Movement Modifier"))
 	CHAOSMOVER_API FMovementModifierHandle K2_QueueMovementModifier(UPARAM(DisplayName = "Movement Modifier") const int32& MoveAsRawData);
 	DECLARE_FUNCTION(execK2_QueueMovementModifier);
@@ -216,10 +215,17 @@ protected:
 	UWorld* World = nullptr;
 
 private:
+	void UpdateEvents(const FMoverTimeStep& TimeStep, UE::ChaosMover::FSimulationOutputData& OutputData);
+
 	FMoverInputCmdContext NetInputCmd;
 	FMoverSyncState NetSyncState;
+
+	// Events from this simulation update
 	TArray<TSharedPtr<FMoverSimulationEventData>> Events;
 
+	// Events that have been sent to GT but are still within the resimulation window
+	TArray<TSharedPtr<FMoverSimulationEventData>> ProcessedEvents;
+
 	bool bInputCmdOverridden = false;
 	bool bSyncStateOverridden = false;
 
diff --git a/Engine/Plugins/Experimental/ChaosMover/Source/ChaosMover/Public/ChaosMover/ChaosMoverSimulationTypes.h b/Engine/Plugins/Experimental/ChaosMover/Source/ChaosMover/Public/ChaosMover/ChaosMoverSimulationTypes.h
--- a/Engine/Plugins/Experimental/ChaosMover/Source/ChaosMover/Public/ChaosMover/ChaosMoverSimulationTypes.h
+++ b/Engine/Plugins/Experimental/ChaosMover/Source/ChaosMover/Public/ChaosMover/ChaosMoverSimulationTypes.h
@@ -12,8 +12,6 @@
 
 #include "ChaosMoverSimulationTypes.generated.h"
 
-#define UE_API CHAOSMOVER_API
-
 namespace Chaos
 {
 	class FPBDJointSettings;
@@ -21,6 +19,7 @@ namespace Chaos
 }
 
 struct FConstraintProfileProperties;
+class FAsyncNetworkPhysicsComponent;
 class UChaosMoverSimulation;
 class UBaseMovementMode;
 
@@ -109,13 +108,13 @@ struct FChaosMovementBasis : public FMoverDataStructBase
 	GENERATED_BODY()
 
 	// Implementation of FMoverDataStructBase
-	UE_API virtual bool NetSerialize(FArchive& Ar, UPackageMap* Map, bool& bOutSuccess) override;
-	UE_API virtual FMoverDataStructBase* Clone() const override;
-	UE_API virtual UScriptStruct* GetScriptStruct() const override;
-	UE_API virtual void ToString(FAnsiStringBuilderBase& Out) const override;
-	UE_API virtual bool ShouldReconcile(const FMoverDataStructBase& AuthorityState) const override;
-	UE_API virtual void Interpolate(const FMoverDataStructBase& From, const FMoverDataStructBase& To, float Pct) override;
-	UE_API virtual void Merge(const FMoverDataStructBase& From) override;
+	CHAOSMOVER_API virtual bool NetSerialize(FArchive& Ar, UPackageMap* Map, bool& bOutSuccess) override;
+	CHAOSMOVER_API virtual FMoverDataStructBase* Clone() const override;
+	CHAOSMOVER_API virtual UScriptStruct* GetScriptStruct() const override;
+	CHAOSMOVER_API virtual void ToString(FAnsiStringBuilderBase& Out) const override;
+	CHAOSMOVER_API virtual bool ShouldReconcile(const FMoverDataStructBase& AuthorityState) const override;
+	CHAOSMOVER_API virtual void Interpolate(const FMoverDataStructBase& From, const FMoverDataStructBase& To, float Pct) override;
+	CHAOSMOVER_API virtual void Merge(const FMoverDataStructBase& From) override;
 
 	UPROPERTY(VisibleAnywhere, BlueprintReadWrite, Category = ChaosMover)
 	FVector BasisLocation = FVector::ZeroVector;
@@ -135,7 +134,7 @@ struct FChaosMoverSimulationDefaultInputs : public FMoverDataStructBase
 		Reset();
 	}
 
-	UE_API void Reset();
+	CHAOSMOVER_API void Reset();
 
 	/** return newly allocated copy of this FMoverDataStructBase. MUST be overridden by derived types. */
 	virtual FMoverDataStructBase* Clone() const override
@@ -169,6 +168,7 @@ struct FChaosMoverSimulationDefaultInputs : public FMoverDataStructBase
 	Chaos::FPhysicsObjectHandle PhysicsObject;
 	AActor* OwningActor;
 	UWorld* World;
+	FAsyncNetworkPhysicsComponent* AsyncNetworkPhysicsComponent;
 
 	UPROPERTY(VisibleAnywhere, BlueprintReadOnly, Category=ChaosMover)
 	float PhysicsObjectGravity;
@@ -291,6 +291,19 @@ struct FStanceModifiedEventData : public FMoverSimulationEventData
 		return FStanceModifiedEventData::StaticStruct();
 	}
 
+	virtual bool IsEqual(const FMoverSimulationEventData& Other, int32 ServerFrameDiffTolerance = 0) const override
+	{
+		if (Super::IsEqual(Other, ServerFrameDiffTolerance))
+		{
+			const FStanceModifiedEventData& TypedOther = static_cast<const FStanceModifiedEventData&>(Other);
+			return (OldStance == TypedOther.OldStance) && (NewStance == TypedOther.NewStance);
+		}
+		else
+		{
+			return false;
+		}
+	}
+
 	EStanceMode OldStance = EStanceMode::Invalid;
 	EStanceMode NewStance = EStanceMode::Invalid;
 };
@@ -313,6 +326,19 @@ struct FStanceModifierModeChangedEventData : public FMoverSimulationEventData
 		return FStanceModifierModeChangedEventData::StaticStruct();
 	}
 
+	virtual bool IsEqual(const FMoverSimulationEventData& Other, int32 ServerFrameDiffTolerance = 0) const override
+	{
+		if (Super::IsEqual(Other, ServerFrameDiffTolerance))
+		{
+			const FStanceModifierModeChangedEventData& TypedOther = static_cast<const FStanceModifierModeChangedEventData&>(Other);
+			return (OldMode == TypedOther.OldMode) && (NewMode == TypedOther.NewMode);
+		}
+		else
+		{
+			return false;
+		}
+	}
+
 	FName OldMode = NAME_None;
 	FName NewMode = NAME_None;
 };
@@ -402,18 +428,18 @@ struct FAsyncLocalOnlyInstantMovementEffect : public FInstantMovementEffect
 {
 	GENERATED_BODY()
 
-	UE_API FAsyncLocalOnlyInstantMovementEffect() {}
+	CHAOSMOVER_API FAsyncLocalOnlyInstantMovementEffect() {}
 	virtual ~FAsyncLocalOnlyInstantMovementEffect() = default;
 
 	FName OptionalName;
 	// Async Function executed in ApplyMovementEffect_Async (only locally)
 	TFunction<bool(FApplyMovementEffectParams_Async&, FMoverSyncState&)> AsyncFunction;
 
-	UE_API virtual bool ApplyMovementEffect_Async(FApplyMovementEffectParams_Async& ApplyEffectParams, FMoverSyncState& OutputState) override;
-	UE_API virtual FInstantMovementEffect* Clone() const override;
-	UE_API virtual void NetSerialize(FArchive& Ar) override;
-	UE_API virtual UScriptStruct* GetScriptStruct() const override;
-	UE_API virtual FString ToSimpleString() const override;
+	CHAOSMOVER_API virtual bool ApplyMovementEffect_Async(FApplyMovementEffectParams_Async& ApplyEffectParams, FMoverSyncState& OutputState) override;
+	CHAOSMOVER_API virtual FInstantMovementEffect* Clone() const override;
+	CHAOSMOVER_API virtual void NetSerialize(FArchive& Ar) override;
+	CHAOSMOVER_API virtual UScriptStruct* GetScriptStruct() const override;
+	CHAOSMOVER_API virtual FString ToSimpleString() const override;
 };
 
 // This instant movement effect is ChaosMover only and teleport to a location without any checks whatsoever, this is meant for debugging only
@@ -422,18 +448,16 @@ struct FDebugTeleportToInstantMovementEffect : public FInstantMovementEffect
 {
 	GENERATED_BODY()
 
-	UE_API FDebugTeleportToInstantMovementEffect() {}
+	CHAOSMOVER_API FDebugTeleportToInstantMovementEffect() {}
 	virtual ~FDebugTeleportToInstantMovementEffect() {}
 
 	// A component can be set to MAX_FLT to signify "keep that current component unchanged"
 	FVector TeleportLocation = FVector(MAX_FLT, MAX_FLT, MAX_FLT);
 	FRotator TeleportRotation = FRotator(MAX_FLT, MAX_FLT, MAX_FLT);
 
-	UE_API virtual bool ApplyMovementEffect_Async(FApplyMovementEffectParams_Async& ApplyEffectParams, FMoverSyncState& OutputState) override;
-	UE_API virtual FInstantMovementEffect* Clone() const override;
-	UE_API virtual void NetSerialize(FArchive& Ar) override;
-	UE_API virtual UScriptStruct* GetScriptStruct() const override;
-	UE_API virtual FString ToSimpleString() const override;
+	CHAOSMOVER_API virtual bool ApplyMovementEffect_Async(FApplyMovementEffectParams_Async& ApplyEffectParams, FMoverSyncState& OutputState) override;
+	CHAOSMOVER_API virtual FInstantMovementEffect* Clone() const override;
+	CHAOSMOVER_API virtual void NetSerialize(FArchive& Ar) override;
+	CHAOSMOVER_API virtual UScriptStruct* GetScriptStruct() const override;
+	CHAOSMOVER_API virtual FString ToSimpleString() const override;
 };
-
-#undef UE_API
diff --git a/Engine/Plugins/Experimental/Mover/Source/Mover/Private/MoverSimulationTypes.cpp b/Engine/Plugins/Experimental/Mover/Source/Mover/Private/MoverSimulationTypes.cpp
--- a/Engine/Plugins/Experimental/Mover/Source/Mover/Private/MoverSimulationTypes.cpp
+++ b/Engine/Plugins/Experimental/Mover/Source/Mover/Private/MoverSimulationTypes.cpp
@@ -87,23 +87,6 @@ namespace UE::Mover
 		CurrentIndex = (CurrentIndex + 1) % 2;
 		Data[CurrentIndex] = { InTimeStep, InData };
 
-		// When we're resimulating we can be adding data that is at an earlier time to
-		// the currently stored data, so fix up if necessary
-		if (InTimeStep.bIsResimulating)
-		{
-			// Remove previous result if it is ahead in time of this latest result
-			const uint8 PrevIndex = (CurrentIndex + 1) % 2;
-			if (Data[PrevIndex].TimeStep.BaseSimTimeMs > InTimeStep.BaseSimTimeMs)
-			{
-				Data[PrevIndex] = Data[CurrentIndex];
-			}
-
-			// If resimulating find any events stored after the time step and remove
-			Events.RemoveAll([&InTimeStep](const TSharedPtr<FMoverSimulationEventData> Event) {
-				return !Event.IsValid() || (Event->Context.EventTimeMs >= InTimeStep.BaseSimTimeMs);
-				});
-		}
-
 		// Transfer events to local Events array
 		for (TSharedPtr<FMoverSimulationEventData>& Event : Data[CurrentIndex].SimOutputData.Events)
 		{
diff --git a/Engine/Plugins/Experimental/Mover/Source/Mover/Private/MoverTypes.cpp b/Engine/Plugins/Experimental/Mover/Source/Mover/Private/MoverTypes.cpp
--- a/Engine/Plugins/Experimental/Mover/Source/Mover/Private/MoverTypes.cpp
+++ b/Engine/Plugins/Experimental/Mover/Source/Mover/Private/MoverTypes.cpp
@@ -42,6 +42,7 @@ FMoverEventContext::FMoverEventContext()
 
 FMoverEventContext::FMoverEventContext(const FMoverTimeStep& InEventTime)
 	: EventTimeMs(InEventTime.BaseSimTimeMs)
+	, ServerFrame(InEventTime.ServerFrame)
 	, bIsDuringResimulation(InEventTime.bIsResimulating)
 {
 }
diff --git a/Engine/Plugins/Experimental/Mover/Source/Mover/Public/MoverSimulationTypes.h b/Engine/Plugins/Experimental/Mover/Source/Mover/Public/MoverSimulationTypes.h
--- a/Engine/Plugins/Experimental/Mover/Source/Mover/Public/MoverSimulationTypes.h
+++ b/Engine/Plugins/Experimental/Mover/Source/Mover/Public/MoverSimulationTypes.h
@@ -552,6 +552,13 @@ struct FMoverSimulationEventData
 		EventProcessedCallback = Callback;
 	}
 
+	virtual bool IsEqual(const FMoverSimulationEventData& Other, int32 ServerFrameDiffTolerance = 0) const
+	{
+		return (GetScriptStruct() == Other.GetScriptStruct()) &&
+			(Context.ServerFrame <= Other.Context.ServerFrame + ServerFrameDiffTolerance) &&
+			(Context.ServerFrame >= Other.Context.ServerFrame - ServerFrameDiffTolerance);
+	}
+
 private:
 	// This callback is fired when the event is processed on the game thread
 	// This is called before and in addition to any type based handling
@@ -579,6 +586,19 @@ struct FMovementModeChangedEventData : public FMoverSimulationEventData
 		return FMovementModeChangedEventData::StaticStruct();
 	}
 
+	virtual bool IsEqual(const FMoverSimulationEventData& Other, int32 ServerFrameDiffTolerance = 0) const override
+	{
+		if (Super::IsEqual(Other, ServerFrameDiffTolerance))
+		{
+			const FMovementModeChangedEventData& TypedOther = static_cast<const FMovementModeChangedEventData&>(Other);
+			return (PreviousModeName == TypedOther.PreviousModeName) && (NewModeName == TypedOther.NewModeName);
+		}
+		else
+		{
+			return false;
+		}
+	}
+
 	FName PreviousModeName = NAME_None;
 	FName NewModeName = NAME_None;
 };
@@ -603,6 +623,20 @@ struct FTeleportSucceededEventData : public FMoverSimulationEventData
 		return FTeleportSucceededEventData::StaticStruct();
 	}
 
+	virtual bool IsEqual(const FMoverSimulationEventData& Other, int32 ServerFrameDiffTolerance = 0) const override
+	{
+		if (Super::IsEqual(Other, ServerFrameDiffTolerance))
+		{
+			const FTeleportSucceededEventData& TypedOther = static_cast<const FTeleportSucceededEventData&>(Other);
+			return FromLocation.Equals(TypedOther.FromLocation) && FromRotation.Equals(TypedOther.FromRotation) &&
+				ToLocation.Equals(TypedOther.ToLocation) && ToRotation.Equals(TypedOther.ToRotation);
+		}
+		else
+		{
+			return false;
+		}
+	}
+
 	FVector FromLocation;
 	FQuat FromRotation;
 	FVector ToLocation;
@@ -636,6 +670,20 @@ struct FTeleportFailedEventData : public FMoverSimulationEventData
 		return FTeleportFailedEventData::StaticStruct();
 	}
 
+	virtual bool IsEqual(const FMoverSimulationEventData& Other, int32 ServerFrameDiffTolerance = 0) const override
+	{
+		if (Super::IsEqual(Other, ServerFrameDiffTolerance))
+		{
+			const FTeleportFailedEventData& TypedOther = static_cast<const FTeleportFailedEventData&>(Other);
+			return FromLocation.Equals(TypedOther.FromLocation) && FromRotation.Equals(TypedOther.FromRotation) &&
+				ToLocation.Equals(TypedOther.ToLocation) && ToRotation.Equals(TypedOther.ToRotation);
+		}
+		else
+		{
+			return false;
+		}
+	}
+
 	FVector FromLocation;
 	FQuat FromRotation;
 	FVector ToLocation;
diff --git a/Engine/Plugins/Experimental/Mover/Source/Mover/Public/MoverTypes.h b/Engine/Plugins/Experimental/Mover/Source/Mover/Public/MoverTypes.h
--- a/Engine/Plugins/Experimental/Mover/Source/Mover/Public/MoverTypes.h
+++ b/Engine/Plugins/Experimental/Mover/Source/Mover/Public/MoverTypes.h
@@ -94,6 +94,10 @@ struct FMoverEventContext
 	UPROPERTY(EditAnywhere, BlueprintReadWrite, Category = Mover)
 	double EventTimeMs = -1.0;
 
+	// Server frame that this event occurred
+	UPROPERTY(EditAnywhere, BlueprintReadWrite, Category = Mover)
+	int32 ServerFrame = -1;
+
 	// Did this event occur while resimulating a past sim frame?
 	UPROPERTY(EditAnywhere, BlueprintReadWrite, Category = Mover)
 	uint8 bIsDuringResimulation : 1 = 0;
diff --git a/Engine/Source/Runtime/Engine/Private/PhysicsEngine/NetworkPhysicsComponent.cpp b/Engine/Source/Runtime/Engine/Private/PhysicsEngine/NetworkPhysicsComponent.cpp
--- a/Engine/Source/Runtime/Engine/Private/PhysicsEngine/NetworkPhysicsComponent.cpp
+++ b/Engine/Source/Runtime/Engine/Private/PhysicsEngine/NetworkPhysicsComponent.cpp
@@ -2867,6 +2867,15 @@ const int32 FAsyncNetworkPhysicsComponent::SetupRewindData()
 	return NumFrames;
 }
 
+int32 FAsyncNetworkPhysicsComponent::GetLastReceivedServerFrame() const
+{
+	if (StateHistory)
+	{
+		return StateHistory->GetLatestFrame();
+	}
+	return -1;
+}
+
 namespace UE::NetworkPhysicsUtils
 {
 	int32 GetUpcomingServerFrame_External(UWorld* World)
diff --git a/Engine/Source/Runtime/Engine/Public/Physics/NetworkPhysicsComponent.h b/Engine/Source/Runtime/Engine/Public/Physics/NetworkPhysicsComponent.h
--- a/Engine/Source/Runtime/Engine/Public/Physics/NetworkPhysicsComponent.h
+++ b/Engine/Source/Runtime/Engine/Public/Physics/NetworkPhysicsComponent.h
@@ -1927,6 +1927,9 @@ public:
 	// Enable RewindData history caching and return the history size
 	const int32 SetupRewindData();
 
+	// Get the frame number of the last server update
+	ENGINE_API int32 GetLastReceivedServerFrame() const;
+
 private:
 	// Initialize, bind delegates etc.
 	void OnInitialize_Internal();
