[Chaos Mover] Add CVar to allow disabling of system to prevent repeat firing of events during resimulation

diff --git a/Engine/Plugins/Experimental/ChaosMover/Source/ChaosMover/Private/ChaosMoverConsoleVariables.cpp b/Engine/Plugins/Experimental/ChaosMover/Source/ChaosMover/Private/ChaosMoverConsoleVariables.cpp
--- a/Engine/Plugins/Experimental/ChaosMover/Source/ChaosMover/Private/ChaosMoverConsoleVariables.cpp
+++ b/Engine/Plugins/Experimental/ChaosMover/Source/ChaosMover/Private/ChaosMoverConsoleVariables.cpp
@@ -34,6 +34,10 @@ namespace UE::ChaosMover
 		bool bSkipGenerateMoveIfOverridden = true;
 		FAutoConsoleVariableRef CVarChaosMoverSkipGenerateMoveIfOverridden(TEXT("ChaosMover.Perf.SkipGenerateMoveIfOverridden"),
 			bSkipGenerateMoveIfOverridden, TEXT("If true and we have a layered move fully overriding movement, then we will skip calling OnGenerateMove on the active movement mode for better performance\n"));
+		
+		bool bDisableResimDuplicateEventChecking = true;
+		FAutoConsoleVariableRef CVarChaosMoverDisableResimDuplicateEventChecking(TEXT("ChaosMover.Networking.DisableResimDuplicateEventChecking"),
+			bDisableResimDuplicateEventChecking, TEXT("If true we disable the check for duplicate events in resim and can resend the same events\n"));
 
 		int32 InstantMovementEffectIDHistorySize = 30;
 		FAutoConsoleVariableRef CVarChaosMoverInstantMovementEffectIDHistorySize(TEXT("ChaosMover.Networking.InstantMovementEffectIDHistorySize"),
diff --git a/Engine/Plugins/Experimental/ChaosMover/Source/ChaosMover/Private/ChaosMoverSimulation.cpp b/Engine/Plugins/Experimental/ChaosMover/Source/ChaosMover/Private/ChaosMoverSimulation.cpp
--- a/Engine/Plugins/Experimental/ChaosMover/Source/ChaosMover/Private/ChaosMoverSimulation.cpp
+++ b/Engine/Plugins/Experimental/ChaosMover/Source/ChaosMover/Private/ChaosMoverSimulation.cpp
@@ -844,48 +844,51 @@ void UChaosMoverSimulation::OnPostSimulationTick(const FMoverTimeStep& TimeStep,
 
 void UChaosMoverSimulation::UpdateEvents(const FMoverTimeStep& TimeStep, UE::ChaosMover::FSimulationOutputData& OutputData)
 {
-	if (const FChaosMoverSimulationDefaultInputs* SimInputs = LocalSimInput.FindDataByType<FChaosMoverSimulationDefaultInputs>())
+	if (!UE::ChaosMover::CVars::bDisableResimDuplicateEventChecking)
 	{
-		// If networking is enabled then we can resimulate and potentially fire the same event
-		// multiple times. To avoid this we maintain a list of processed events and if the same
-		// event is generated at the same time then remove it
-		if (SimInputs->AsyncNetworkPhysicsComponent && !SimInputs->AsyncNetworkPhysicsComponent->IsServer())
+		if (const FChaosMoverSimulationDefaultInputs* SimInputs = LocalSimInput.FindDataByType<FChaosMoverSimulationDefaultInputs>())
 		{
-			// Clear events that are outside of the resimulation window
-			const int32 LastReceivedServerFrame = SimInputs->AsyncNetworkPhysicsComponent->GetLastReceivedServerFrame();
-			if (LastReceivedServerFrame > 0)
+			// If networking is enabled then we can resimulate and potentially fire the same event
+			// multiple times. To avoid this we maintain a list of processed events and if the same
+			// event is generated at the same time then remove it
+			if (SimInputs->AsyncNetworkPhysicsComponent && !SimInputs->AsyncNetworkPhysicsComponent->IsServer())
 			{
-				ProcessedEvents.RemoveAllSwap([&LastReceivedServerFrame](const TSharedPtr<FMoverSimulationEventData> Event) {
-					return Event->Context.ServerFrame <= LastReceivedServerFrame;
-					});
-			}
+				// Clear events that are outside of the resimulation window
+				const int32 LastReceivedServerFrame = SimInputs->AsyncNetworkPhysicsComponent->GetLastReceivedServerFrame();
+				if (LastReceivedServerFrame > 0)
+				{
+					ProcessedEvents.RemoveAllSwap([&LastReceivedServerFrame](const TSharedPtr<FMoverSimulationEventData> Event) {
+						return Event->Context.ServerFrame <= LastReceivedServerFrame;
+						});
+				}
 
-			if (TimeStep.bIsResimulating)
-			{
-				// Remove events that have already been processed in the main update
+				if (TimeStep.bIsResimulating)
+				{
+					// Remove events that have already been processed in the main update
 
-				Events.RemoveAll([this](const TSharedPtr<FMoverSimulationEventData> Event) {
-					const int32 AllowedDiff = UE::ChaosMover::CVars::FrameDifferenceLeniencyForEventTimeComparison;
-					for (TSharedPtr<FMoverSimulationEventData>& ProcessedEvent : ProcessedEvents)
-					{
-						if (const FMoverSimulationEventData* ProcessedEventPtr = ProcessedEvent.Get())
+					Events.RemoveAll([this](const TSharedPtr<FMoverSimulationEventData> Event) {
+						const int32 AllowedDiff = UE::ChaosMover::CVars::FrameDifferenceLeniencyForEventTimeComparison;
+						for (TSharedPtr<FMoverSimulationEventData>& ProcessedEvent : ProcessedEvents)
 						{
-							if (Event->IsEqual(*ProcessedEventPtr, AllowedDiff))
+							if (const FMoverSimulationEventData* ProcessedEventPtr = ProcessedEvent.Get())
 							{
-								return true;
+								if (Event->IsEqual(*ProcessedEventPtr, AllowedDiff))
+								{
+									return true;
+								}
 							}
+
 						}
-						
-					}
-					return false;
-					});
-			}
-			else
-			{
-				// Add new events to the processed list
-				for (TSharedPtr<FMoverSimulationEventData>& Event : Events)
+						return false;
+						});
+				}
+				else
 				{
-					ProcessedEvents.Add(Event);
+					// Add new events to the processed list
+					for (TSharedPtr<FMoverSimulationEventData>& Event : Events)
+					{
+						ProcessedEvents.Add(Event);
+					}
 				}
 			}
 		}
diff --git a/Engine/Plugins/Experimental/ChaosMover/Source/ChaosMover/Public/ChaosMover/ChaosMoverConsoleVariables.h b/Engine/Plugins/Experimental/ChaosMover/Source/ChaosMover/Public/ChaosMover/ChaosMoverConsoleVariables.h
--- a/Engine/Plugins/Experimental/ChaosMover/Source/ChaosMover/Public/ChaosMover/ChaosMoverConsoleVariables.h
+++ b/Engine/Plugins/Experimental/ChaosMover/Source/ChaosMover/Public/ChaosMover/ChaosMoverConsoleVariables.h
@@ -13,6 +13,7 @@ namespace UE::ChaosMover
 		CHAOSMOVER_API extern bool bDrawGroundQueries;
 		CHAOSMOVER_API extern bool bDrawOverlapQueries;
 		CHAOSMOVER_API extern bool bSkipGenerateMoveIfOverridden;
+		CHAOSMOVER_API extern bool bDisableResimDuplicateEventChecking;
 		CHAOSMOVER_API extern int32 InstantMovementEffectIDHistorySize;
 		CHAOSMOVER_API extern int32 FrameDifferenceLeniencyForEventTimeComparison;
 	}
