[ChaosMover] Setup and validation to make it easier to start using ChaosMover

diff --git a/Engine/Plugins/Experimental/ChaosMover/Config/DefaultChaosMover.ini b/Engine/Plugins/Experimental/ChaosMover/Config/DefaultChaosMover.ini
--- a/Engine/Plugins/Experimental/ChaosMover/Config/DefaultChaosMover.ini
+++ b/Engine/Plugins/Experimental/ChaosMover/Config/DefaultChaosMover.ini
@@ -1,2 +1,5 @@
 ﻿[CoreRedirects]
-+PropertyRedirects=(OldName="/Script/ChaosMover.ChaosCharacterMoverComponent.OnLanded",NewName="/Script/ChaosMover.ChaosCharacterMoverComponent.OnLandedDelegate")
\ No newline at end of file
++PropertyRedirects=(OldName="/Script/ChaosMover.ChaosCharacterMoverComponent.OnLanded",NewName="/Script/ChaosMover.ChaosCharacterMoverComponent.OnLandedDelegate")
+
+[/Script/Engine.PhysicsSettings]
+PhysicsPrediction=(bEnablePhysicsPrediction=True)
\ No newline at end of file
diff --git a/Engine/Plugins/Experimental/ChaosMover/Source/ChaosMover/Private/Backends/ChaosMoverBackend.cpp b/Engine/Plugins/Experimental/ChaosMover/Source/ChaosMover/Private/Backends/ChaosMoverBackend.cpp
--- a/Engine/Plugins/Experimental/ChaosMover/Source/ChaosMover/Private/Backends/ChaosMoverBackend.cpp
+++ b/Engine/Plugins/Experimental/ChaosMover/Source/ChaosMover/Private/Backends/ChaosMoverBackend.cpp
@@ -72,46 +72,35 @@ void UChaosMoverBackendComponent::InitializeComponent()
 		Simulation->SetRollbackBlackboard(GetMoverComponent().GetRollbackBlackboard_Internal());
 
 		// Create NetworkPhysicsComponent
-		if ((World->GetNetMode() != ENetMode::NM_Standalone) && Chaos::FPhysicsSolverBase::IsNetworkPhysicsPredictionEnabled())
+		if ((World->GetNetMode() != ENetMode::NM_Standalone) && Chaos::FPhysicsSolverBase::IsNetworkPhysicsPredictionEnabled() && bIsUsingAsyncPhysics)
 		{
-			if (!bIsUsingAsyncPhysics)
+			NetworkPhysicsComponent = NewObject<UNetworkPhysicsComponent>(GetOwner(), TEXT("PhysMover_NetworkPhysicsComponent"), RF_Transient);
+
+			// This isn't technically a DSO component, but set it net addressable as though it is
+			NetworkPhysicsComponent->SetNetAddressable();
+			NetworkPhysicsComponent->SetIsReplicated(true);
+			NetworkPhysicsComponent->RegisterComponent();
+			if (!NetworkPhysicsComponent->HasBeenInitialized())
 			{
-				// Verify that the Project Settings have bTickPhysicsAsync turned on.
-				// It's easy to waste time forgetting that, since they are off by default.
-				UE_LOG(LogChaosMover, Warning, TEXT("Chaos Mover Backend only supports networking with Physics Async. Networked Physics will not work well. Turn on 'Project Settings > Engine - Physics > Tick Physics Async', or play in Standalone Mode"));
-				// This is important enough that we break for developers debugging in editor
-				UE_DEBUG_BREAK();
+				NetworkPhysicsComponent->InitializeComponent();
 			}
-			else
-			{
-				NetworkPhysicsComponent = NewObject<UNetworkPhysicsComponent>(GetOwner(), TEXT("PhysMover_NetworkPhysicsComponent"), RF_Transient);
+			NetworkPhysicsComponent->Activate(true);
+
+			// Register network data for recording and rewind/resim
+			NetworkPhysicsComponent->CreateDataHistory<UE::ChaosMover::FNetworkDataTraits>(this);
 
-				// This isn't technically a DSO component, but set it net addressable as though it is
-				NetworkPhysicsComponent->SetNetAddressable();
-				NetworkPhysicsComponent->SetIsReplicated(true);
-				NetworkPhysicsComponent->RegisterComponent();
-				if (!NetworkPhysicsComponent->HasBeenInitialized())
+			if (NetworkPhysicsComponent->HasServerWorld())
+			{
+				if (APawn* PawnOwner = Cast<APawn>(GetOwner()))
 				{
-					NetworkPhysicsComponent->InitializeComponent();
+					// When we're owned by a pawn, keep an eye on whether it's currently player-controlled or not
+					PawnOwner->ReceiveControllerChangedDelegate.AddUniqueDynamic(this, &ThisClass::HandleOwningPawnControllerChanged_Server);
+					HandleOwningPawnControllerChanged_Server(PawnOwner, nullptr, PawnOwner->Controller);
 				}
-				NetworkPhysicsComponent->Activate(true);
-
-				// Register network data for recording and rewind/resim
-				NetworkPhysicsComponent->CreateDataHistory<UE::ChaosMover::FNetworkDataTraits>(this);
-
-				if (NetworkPhysicsComponent->HasServerWorld())
+				else
 				{
-					if (APawn* PawnOwner = Cast<APawn>(GetOwner()))
-					{
-						// When we're owned by a pawn, keep an eye on whether it's currently player-controlled or not
-						PawnOwner->ReceiveControllerChangedDelegate.AddUniqueDynamic(this, &ThisClass::HandleOwningPawnControllerChanged_Server);
-						HandleOwningPawnControllerChanged_Server(PawnOwner, nullptr, PawnOwner->Controller);
-					}
-					else
-					{
-						// If the owner isn't a pawn, there's no chance of player input happening, so inputs to the PT are always produced on the server
-						NetworkPhysicsComponent->SetIsRelayingLocalInputs(true);
-					}
+					// If the owner isn't a pawn, there's no chance of player input happening, so inputs to the PT are always produced on the server
+					NetworkPhysicsComponent->SetIsRelayingLocalInputs(true);
 				}
 			}
 		}
@@ -376,6 +365,14 @@ void UChaosMoverBackendComponent::BeginPlay()
 
 	if (UWorld* World = GetWorld(); World && World->IsGameWorld())
 	{
+		if (UPrimitiveComponent* UpdatedPrimitiveComp = Cast<UPrimitiveComponent>(GetMoverComponent().GetUpdatedComponent()))
+		{
+			if (!UpdatedPrimitiveComp->IsSimulatingPhysics())
+			{
+				UE_LOG(LogChaosMover, Warning, TEXT("ChaosMoverBackend: Updated primitive component '%s' on actor '%s' does not have Simulate Physics enabled. ChaosMover requires updating the primitive component through physics, so 'Simulate Physics' must be enabled on the updated primitive component."), *GetNameSafe(UpdatedPrimitiveComp), *GetNameSafe(GetOwner()));
+			}
+		}
+
 		CreatePhysics();
 
 		// Register with the world subsystem
diff --git a/Engine/Plugins/Experimental/ChaosMover/Source/ChaosMover/Private/Backends/ChaosMoverSubsystem.cpp b/Engine/Plugins/Experimental/ChaosMover/Source/ChaosMover/Private/Backends/ChaosMoverSubsystem.cpp
--- a/Engine/Plugins/Experimental/ChaosMover/Source/ChaosMover/Private/Backends/ChaosMoverSubsystem.cpp
+++ b/Engine/Plugins/Experimental/ChaosMover/Source/ChaosMover/Private/Backends/ChaosMoverSubsystem.cpp
@@ -6,10 +6,12 @@
 
 #include "ChaosMover/Backends/ChaosMoverBackend.h"
 #include "ChaosMover/ChaosMoverConsoleVariables.h"
+#include "ChaosMover/ChaosMoverLog.h"
 #include "GameFramework/PlayerController.h"
 #include "PBDRigidsSolver.h"
 #include "Physics/Experimental/PhysScene_Chaos.h"
 #include "Physics/NetworkPhysicsComponent.h"
+#include "PhysicsEngine/PhysicsSettings.h"
 
 #include UE_INLINE_GENERATED_CPP_BY_NAME(ChaosMoverSubsystem)
 
@@ -72,6 +74,38 @@ void UChaosMoverSubsystem::Deinitialize()
 
 void UChaosMoverSubsystem::Register(TWeakObjectPtr<UChaosMoverBackendComponent> InBackend)
 {
+	if (Backends.Num() == 0 && InBackend.IsValid())
+	{
+		if (UWorld* World = GetWorld())
+		{
+			UPhysicsSettings* PhysicsSettings = UPhysicsSettings::Get();
+			bool bProjectSettingAsyncEnabled = PhysicsSettings ? PhysicsSettings->bTickPhysicsAsync : false;
+			bool bProjectSettingPhysicsPredictionEnabled = PhysicsSettings ? PhysicsSettings->PhysicsPrediction.bEnablePhysicsPrediction : false;
+
+			if (PhysicsSettings)
+			{
+				if (World->GetNetMode() != ENetMode::NM_Standalone)
+				{
+					if (!bProjectSettingAsyncEnabled)
+					{
+						UE_LOG(LogChaosMover, Warning, TEXT("ChaosMoverSubsystem: bTickPhysicsAsync project setting is disabled. bTickPhysicsAsync is required for networked physics to replicate correctly. Enable 'Tick Physics Async' in Project Settings > Engine > Physics."));
+#if !UE_BUILD_SHIPPING
+						UE_DEBUG_BREAK();
+#endif
+					}
+				}
+
+				if (!bProjectSettingPhysicsPredictionEnabled)
+				{
+					UE_LOG(LogChaosMover, Warning, TEXT("ChaosMoverSubsystem: bEnablePhysicsPrediction project setting is disabled. Enable 'Enable Physics Prediction' in Project Settings > Engine > Physics."));
+#if !UE_BUILD_SHIPPING
+					UE_DEBUG_BREAK();
+#endif
+				}
+			}
+		}
+	}
+
 	Backends.AddUnique(InBackend);
 }
 
