ChaosMover: Tracing Instant Movement Effects into Chaos Visual Debugger

#Synced-CL 48072262

diff --git a/Engine/Plugins/Experimental/ChaosMover/Source/ChaosMover/Private/ChaosMoverSimulation.cpp b/Engine/Plugins/Experimental/ChaosMover/Source/ChaosMover/Private/ChaosMoverSimulation.cpp
--- a/Engine/Plugins/Experimental/ChaosMover/Source/ChaosMover/Private/ChaosMoverSimulation.cpp
+++ b/Engine/Plugins/Experimental/ChaosMover/Source/ChaosMover/Private/ChaosMoverSimulation.cpp
@@ -934,6 +934,12 @@ void UChaosMoverSimulation::TraceMoverData(const FMoverTimeStep& TimeStep, const
 
 		// Clear debug sim data to avoid recording past debug data as current on the next frame
 		DebugSimData.Empty();
+
+		// Let the state machine know we have traced data. From this point on, the state machine no longer records instant movement effects directly into DebugSimData
+		// but instead will log all its unprocessed instant movement effects into DebugSimData at the beginning of the next OnSimulationTick. This allows
+		// us to capture both the instant movement effects that were added outside of SimulationTick AND any scheduled movement effects which are waiting for their
+		// execution server frame to arrive
+		StateMachine.OnEndTraceMoverData();
 	}
 #endif
 }
diff --git a/Engine/Plugins/Experimental/ChaosMover/Source/ChaosMover/Private/ChaosMoverSimulationTypes.cpp b/Engine/Plugins/Experimental/ChaosMover/Source/ChaosMover/Private/ChaosMoverSimulationTypes.cpp
--- a/Engine/Plugins/Experimental/ChaosMover/Source/ChaosMover/Private/ChaosMoverSimulationTypes.cpp
+++ b/Engine/Plugins/Experimental/ChaosMover/Source/ChaosMover/Private/ChaosMoverSimulationTypes.cpp
@@ -91,12 +91,25 @@ UScriptStruct* FChaosMoverTimeStepDebugData::GetScriptStruct() const
 
 void FChaosNetInstantMovementEffectsQueue::Add(const FScheduledInstantMovementEffect& InScheduledEffect, int32 IssuanceServerFrame, uint8 UniqueID)
 {
-	FChaosNetInstantMovementEffect& AddedInstancedEffect = Effects.AddDefaulted_GetRef();
-	AddedInstancedEffect.ExecutionServerFrame = InScheduledEffect.ExecutionServerFrame;
-	AddedInstancedEffect.IssuanceServerFrame = IssuanceServerFrame;	
-	AddedInstancedEffect.UniqueID = UniqueID;
 	TCheckedObjPtr<FInstantMovementEffect> InstantMovementEffect = InScheduledEffect.Effect.Get();
-	AddedInstancedEffect.Effect.InitializeAsScriptStruct(InstantMovementEffect->GetScriptStruct(), (const uint8*)InScheduledEffect.Effect.Get());
+	if (ensure(InstantMovementEffect.IsValid()))
+	{
+		FChaosNetInstantMovementEffect& AddedInstancedEffect = Effects.AddDefaulted_GetRef();
+		AddedInstancedEffect.ExecutionServerFrame = InScheduledEffect.ExecutionServerFrame;
+		AddedInstancedEffect.IssuanceServerFrame = IssuanceServerFrame;	
+		AddedInstancedEffect.UniqueID = UniqueID;
+		AddedInstancedEffect.Effect.InitializeAsScriptStruct(InstantMovementEffect->GetScriptStruct(), (const uint8*)InScheduledEffect.Effect.Get());
+	}
+}
+
+FChaosNetInstantMovementEffectsQueue& FChaosNetInstantMovementEffectsQueue::operator=(const TArray<FChaosScheduledInstantMovementEffect>& ScheduledInstantEffects)
+{
+	Effects.Empty(/*Slack =*/ ScheduledInstantEffects.Num());
+	for (const FChaosScheduledInstantMovementEffect& ScheduledInstantEffect : ScheduledInstantEffects)
+	{
+		Add(ScheduledInstantEffect.ScheduledEffect, ScheduledInstantEffect.IssuanceServerFrame, 0xFF/* This is used for debugging, unique ID doesn't apply */);
+	}
+	return *this;
 }
 
 bool FChaosNetInstantMovementEffectsQueue::NetSerialize(FArchive& Ar, UPackageMap* Map, bool& bOutSuccess)
@@ -279,11 +292,6 @@ FInstantMovementEffect* FAsyncLocalOnlyInstantMovementEffect::Clone() const
 	return new FAsyncLocalOnlyInstantMovementEffect(*this);
 }
 
-void FAsyncLocalOnlyInstantMovementEffect::NetSerialize(FArchive& Ar)
-{
-	// No serialization, this is a local only effect
-}
-
 UScriptStruct* FAsyncLocalOnlyInstantMovementEffect::GetScriptStruct() const
 {
 	return FAsyncLocalOnlyInstantMovementEffect::StaticStruct();
diff --git a/Engine/Plugins/Experimental/ChaosMover/Source/ChaosMover/Private/ChaosMoverStateMachine.cpp b/Engine/Plugins/Experimental/ChaosMover/Source/ChaosMover/Private/ChaosMoverStateMachine.cpp
--- a/Engine/Plugins/Experimental/ChaosMover/Source/ChaosMover/Private/ChaosMoverStateMachine.cpp
+++ b/Engine/Plugins/Experimental/ChaosMover/Source/ChaosMover/Private/ChaosMoverStateMachine.cpp
@@ -18,6 +18,9 @@
 #include "Framework/Threading.h"
 #include "Engine/World.h"
 
+#if WITH_CHAOS_VISUAL_DEBUGGER
+#include "ChaosVisualDebugger/ChaosVisualDebuggerTrace.h"
+#endif
 namespace UE::ChaosMover
 {
 	FMoverStateMachine::FMoverStateMachine()
@@ -333,6 +336,10 @@ namespace UE::ChaosMover
 			ReceiveInstantMovementEffects(TimeStep, InstantMovementEffectsQueue);
 		}
 
+#if WITH_CHAOS_VISUAL_DEBUGGER
+		CaptureUnprocessedInstantMovementEffectsForCVD();
+#endif
+
 		TStrongObjectPtr<UImmediateMovementModeTransition> QueuedModeTransition = QueuedModeTransitionWeakPtr.Pin();
 		if (QueuedModeTransition && !QueuedModeTransition->IsSet())
 		{
@@ -755,6 +762,10 @@ namespace UE::ChaosMover
 	void FMoverStateMachine::QueueInstantMovementEffect(const FChaosScheduledInstantMovementEffect& ScheduledEffect)
 	{
 		QueuedInstantEffects.Add(ScheduledEffect);
+
+#if WITH_CHAOS_VISUAL_DEBUGGER
+		TraceInstantMovementEffect(ScheduledEffect);
+#endif
 	}
 
 	void FMoverStateMachine::ConstructDefaultModes()
@@ -943,4 +954,44 @@ namespace UE::ChaosMover
 		}
 		return false;
 	}
+
+#if WITH_CHAOS_VISUAL_DEBUGGER
+	void FMoverStateMachine::CaptureUnprocessedInstantMovementEffectsForCVD()
+	{
+		// Record the list of instant movement effects we ended up with at the start of OnSimulationTick
+		// These could be a combination of effects we will process on the current frame or effects that are scheduled to be processed later.
+		// The fact that we have scheduled effects that stay in the queue for a while is the reason we can't simply capture effects for recording in QueueInstantMovementEffect
+		// Some of these effects come from the FChaosNetInstantMovementEffectsQueue found in StartState.InputCmd.InputCollection (transferred in ReceiveInstantMovementEffects)
+		// while others come from direct calls to QueueInstantMovementEffect on the simulation
+		// Since we want to also record the effects that are added during OnSimulationTick after this point, we set a flag so any calls to QueueInstantMovementEffect
+		// from now till post solve when we trace data to CVD are to be added to DebugInstantMovementEffectsQueue
+		if (FChaosVisualDebuggerTrace::IsTracing())
+		{
+			if (TStrongObjectPtr<UChaosMoverSimulation> SimStrongObjPtr = Simulation.Pin())
+			{
+				FChaosNetInstantMovementEffectsQueue& DebugInstantMovementEffectsQueue = SimStrongObjPtr->GetDebugSimData().FindOrAddMutableDataByType<FChaosNetInstantMovementEffectsQueue>();
+				DebugInstantMovementEffectsQueue = QueuedInstantEffects;
+				bCaptureInstantMovementEffectsForCVD = true;
+			}
+		}
+	}
+
+	void FMoverStateMachine::TraceInstantMovementEffect(const FChaosScheduledInstantMovementEffect& ScheduledEffect)
+	{
+		if (bCaptureInstantMovementEffectsForCVD && FChaosVisualDebuggerTrace::IsTracing())
+		{
+			if (TStrongObjectPtr<UChaosMoverSimulation> SimStrongObjPtr = Simulation.Pin())
+			{
+				FChaosNetInstantMovementEffectsQueue& DebugInstantMovementEffectsQueue = SimStrongObjPtr->GetDebugSimData().FindOrAddMutableDataByType<FChaosNetInstantMovementEffectsQueue>();
+				DebugInstantMovementEffectsQueue.Add(ScheduledEffect.ScheduledEffect, ScheduledEffect.IssuanceServerFrame, 0xFF);
+			}
+		}
+	}
+
+	void FMoverStateMachine::OnEndTraceMoverData()
+	{
+		bCaptureInstantMovementEffectsForCVD = false;
+	}
+#endif
+
 } // End of namespace UE::ChaosMover
\ No newline at end of file
diff --git a/Engine/Plugins/Experimental/ChaosMover/Source/ChaosMover/Public/ChaosMover/ChaosMoverSimulationTypes.h b/Engine/Plugins/Experimental/ChaosMover/Source/ChaosMover/Public/ChaosMover/ChaosMoverSimulationTypes.h
--- a/Engine/Plugins/Experimental/ChaosMover/Source/ChaosMover/Public/ChaosMover/ChaosMoverSimulationTypes.h
+++ b/Engine/Plugins/Experimental/ChaosMover/Source/ChaosMover/Public/ChaosMover/ChaosMoverSimulationTypes.h
@@ -410,6 +410,9 @@ struct FChaosNetInstantMovementEffectsQueue : public FMoverDataStructBase
 
 	void Add(const FScheduledInstantMovementEffect& Effect, int32 IssuanceServerFrame, uint8 UniqueID);
 
+	// This function is used for debugging (sending to CVD)
+	FChaosNetInstantMovementEffectsQueue& operator=(const TArray<FChaosScheduledInstantMovementEffect>& ScheduledInstantEffects);
+
 	virtual FMoverDataStructBase* Clone() const override;
 	virtual bool NetSerialize(FArchive& Ar, UPackageMap* Map, bool& bOutSuccess) override;
 	virtual UScriptStruct* GetScriptStruct() const override;
@@ -420,9 +423,8 @@ struct FChaosNetInstantMovementEffectsQueue : public FMoverDataStructBase
 	virtual void Decay(float DecayAmount) override;
 };
 
-
 // This instant movement effect is ChaosMover only and executes AsyncFunction on the issuing end point only
-// AsyncFunction does not (cannot) be serialized, so this is local only
+// Being local only, this instant movement effect does not serialize anything, doesn't override NetSerialize
 USTRUCT()
 struct FAsyncLocalOnlyInstantMovementEffect : public FInstantMovementEffect
 {
@@ -437,7 +439,7 @@ struct FAsyncLocalOnlyInstantMovementEffect : public FInstantMovementEffect
 
 	CHAOSMOVER_API virtual bool ApplyMovementEffect_Async(FApplyMovementEffectParams_Async& ApplyEffectParams, FMoverSyncState& OutputState) override;
 	CHAOSMOVER_API virtual FInstantMovementEffect* Clone() const override;
-	CHAOSMOVER_API virtual void NetSerialize(FArchive& Ar) override;
+	// No NetSerialize override, this is a local only effect
 	CHAOSMOVER_API virtual UScriptStruct* GetScriptStruct() const override;
 	CHAOSMOVER_API virtual FString ToSimpleString() const override;
 };
diff --git a/Engine/Plugins/Experimental/ChaosMover/Source/ChaosMover/Public/ChaosMover/ChaosMoverStateMachine.h b/Engine/Plugins/Experimental/ChaosMover/Source/ChaosMover/Public/ChaosMover/ChaosMoverStateMachine.h
--- a/Engine/Plugins/Experimental/ChaosMover/Source/ChaosMover/Public/ChaosMover/ChaosMoverStateMachine.h
+++ b/Engine/Plugins/Experimental/ChaosMover/Source/ChaosMover/Public/ChaosMover/ChaosMoverStateMachine.h
@@ -72,6 +72,11 @@ namespace UE::ChaosMover
 		void SetOwnerActorName(const FString& InOwnerActorName);
 		void SetOwnerActorLocalNetRole(ENetRole InOwnerActorLocalNetRole);
 
+#if WITH_CHAOS_VISUAL_DEBUGGER
+		void CaptureUnprocessedInstantMovementEffectsForCVD();
+		void OnEndTraceMoverData();
+#endif
+
 	protected:
 
 		CHAOSMOVER_API void QueueNextModeInternal(FName DesiredNextModeName, bool bShouldReenter = false, bool bIsFromRollback=false);
@@ -132,5 +137,14 @@ namespace UE::ChaosMover
 		};
 
 		FIDHistory InstantMovementEffectsIDHistory;
+
+#if WITH_CHAOS_VISUAL_DEBUGGER
+		// If true, queued instant movement effects will be added to DebugSimData
+		// It is false from the last time we traced movement data to the time we call CaptureUnprocessedInstantMovementEffectsForCVD in FMoverStateMachine::OnSimulationTick
+		// which allows us to trace any unprocessed effects, including scheduled ones, that exist on that frame
+		// It is true after that point until we trace, so that we can capture the ones queued after that, for instance those queued by movement modes
+		bool bCaptureInstantMovementEffectsForCVD = false;
+		void TraceInstantMovementEffect(const FChaosScheduledInstantMovementEffect& ScheduledEffect);
+#endif
 	};
 }
\ No newline at end of file
