[Chaos Mover] Add fixes for potential crashes due to accessing deleted objects in the async callback

#Synced-CL 48448165

diff --git a/Engine/Plugins/Experimental/ChaosMover/Source/ChaosMover/Private/ChaosMoverSimulation.cpp b/Engine/Plugins/Experimental/ChaosMover/Source/ChaosMover/Private/ChaosMoverSimulation.cpp
--- a/Engine/Plugins/Experimental/ChaosMover/Source/ChaosMover/Private/ChaosMoverSimulation.cpp
+++ b/Engine/Plugins/Experimental/ChaosMover/Source/ChaosMover/Private/ChaosMoverSimulation.cpp
@@ -66,19 +66,24 @@ FMoverDataCollection& UChaosMoverSimulation::GetDebugSimData()
 	return DebugSimData;
 }
 
-const UBaseMovementMode* UChaosMoverSimulation::GetCurrentMovementMode() const
+TWeakObjectPtr<const UBaseMovementMode> UChaosMoverSimulation::GetCurrentMovementMode() const
 {
-	return StateMachine.GetCurrentMode().Get();
+	return StateMachine.GetCurrentMode();
 }
 
-const UBaseMovementMode* UChaosMoverSimulation::FindMovementModeByName(const FName& Name) const
+TWeakObjectPtr<const UBaseMovementMode> UChaosMoverSimulation::FindMovementModeByName(const FName& Name) const
 {
-	return StateMachine.FindMovementMode(Name).Get();
+	return StateMachine.FindMovementMode(Name);
 }
 
-UBaseMovementMode* UChaosMoverSimulation::FindMovementModeByName_Mutable(const FName& Name)
+TWeakObjectPtr<UBaseMovementMode> UChaosMoverSimulation::FindMovementModeByName_Mutable(const FName& Name)
 {
-	return StateMachine.FindMovementMode_Mutable(Name).Get();
+	return StateMachine.FindMovementMode_Mutable(Name);
+}
+
+bool UChaosMoverSimulation::HasMovementMode(const FName& Name) const
+{
+	return StateMachine.FindMovementMode(Name) != nullptr;
 }
 
 void UChaosMoverSimulation::ApplyNetInputData(const FMoverInputCmdContext& InNetInputCmd)
@@ -1231,7 +1236,7 @@ void UChaosMoverSimulation::CancelModifierFromHandle(FMovementModifierHandle Mod
 bool UChaosMoverSimulation::HasGameplayTag(FGameplayTag TagToFind, bool bExactMatch) const
 {
 	// Check active Movement Mode
-	if (const UBaseMovementMode* ActiveMovementMode = FindMovementModeByName(CurrentSyncState.MovementMode))
+	if (TStrongObjectPtr<const UBaseMovementMode> ActiveMovementMode = FindMovementModeByName(CurrentSyncState.MovementMode).Pin())
 	{
 		if (ActiveMovementMode->HasGameplayTag(TagToFind, bExactMatch))
 		{
diff --git a/Engine/Plugins/Experimental/ChaosMover/Source/ChaosMover/Private/Character/Transitions/ChaosCharacterCrouchCheck.cpp b/Engine/Plugins/Experimental/ChaosMover/Source/ChaosMover/Private/Character/Transitions/ChaosCharacterCrouchCheck.cpp
--- a/Engine/Plugins/Experimental/ChaosMover/Source/ChaosMover/Private/Character/Transitions/ChaosCharacterCrouchCheck.cpp
+++ b/Engine/Plugins/Experimental/ChaosMover/Source/ChaosMover/Private/Character/Transitions/ChaosCharacterCrouchCheck.cpp
@@ -129,7 +129,7 @@ FTransitionEvalResult UChaosCharacterCrouchCheck::Evaluate_Implementation(const
 	const FChaosMoverCrouchInputs* CrouchInputs = TickParams.StartState.InputCmd.InputCollection.FindDataByType<FChaosMoverCrouchInputs>();
 	if (CrouchInputs)
 	{
-		if (const UBaseMovementMode* Mode = Simulation->FindMovementModeByName(TickParams.StartState.SyncState.MovementMode))
+		if (Simulation->HasMovementMode(TickParams.StartState.SyncState.MovementMode))
 		{
 			bool bIsCrouching = Simulation->HasGameplayTag(Mover_IsCrouching, true);
 			if (bIsCrouching && !CrouchModifierHandle.IsValid())
diff --git a/Engine/Plugins/Experimental/ChaosMover/Source/ChaosMover/Private/Character/Transitions/ChaosCharacterFallingCheck.cpp b/Engine/Plugins/Experimental/ChaosMover/Source/ChaosMover/Private/Character/Transitions/ChaosCharacterFallingCheck.cpp
--- a/Engine/Plugins/Experimental/ChaosMover/Source/ChaosMover/Private/Character/Transitions/ChaosCharacterFallingCheck.cpp
+++ b/Engine/Plugins/Experimental/ChaosMover/Source/ChaosMover/Private/Character/Transitions/ChaosCharacterFallingCheck.cpp
@@ -68,56 +68,59 @@ FTransitionEvalResult UChaosCharacterFallingCheck::Evaluate_Implementation(const
 			return EvalResult;
 		}
 
-		if (const IChaosCharacterMovementModeInterface* Mode = Cast<IChaosCharacterMovementModeInterface>(Simulation->FindMovementModeByName(Params.StartState.SyncState.MovementMode)))
+		if (TStrongObjectPtr<const UBaseMovementMode> BaseMode = Simulation->FindMovementModeByName(Params.StartState.SyncState.MovementMode).Pin())
 		{
-			const float DeltaSeconds = Params.TimeStep.StepMs * 0.001f;
-			const FVector LinearVelocity = StartingSyncState->GetVelocity_WorldSpace();
-
-			const FVector ProjectedVelocity = StartingSyncState->GetVelocity_WorldSpace() + SimInputs->Gravity * DeltaSeconds;
-			FVector ProjectedGroundVelocity = UChaosGroundMovementUtils::ComputeLocalGroundVelocity_Internal(StartingSyncState->GetLocation_WorldSpace(), FloorResult);
-			const Chaos::FPBDRigidParticleHandle* GroundParticle = UChaosGroundMovementUtils::GetRigidParticleHandleFromFloorResult_Internal(FloorResult);
-			if (GroundParticle && GroundParticle->IsDynamic() && GroundParticle->GravityEnabled())
-			{
-				// This might not be correct if different physics objects have different gravity but is saves having to go
-				// to the component to get the gravity on the physics volume.
-				ProjectedGroundVelocity += SimInputs->PhysicsObjectGravity * SimInputs->UpDir * DeltaSeconds;
-			}
-			const bool bIsGroundMoving = ProjectedGroundVelocity.SizeSquared() > UE_KINDA_SMALL_NUMBER;
-			const FVector ProjectedRelativeVelocity = ProjectedVelocity - ProjectedGroundVelocity;
-			const float ProjectedRelativeNormalVelocity = FloorResult.HitResult.ImpactNormal.Dot(ProjectedVelocity - ProjectedGroundVelocity);
-			const float ProjectedRelativeVerticalVelocity = SimInputs->UpDir.Dot(ProjectedVelocity - ProjectedGroundVelocity);
-			const float VerticalVelocityLimit = 2.0f / DeltaSeconds;
-			bool bIsLiftingOffSurface = false;
-			if ((ProjectedRelativeNormalVelocity > VerticalVelocityLimit) && bIsGroundMoving && (ProjectedRelativeVerticalVelocity > VerticalVelocityLimit))
-			{
-				bIsLiftingOffSurface = true;
-			}
-
-			const bool bIsWithinReach = FloorResult.FloorDist - Mode->GetTargetHeight() <= SharedSettings->MaxStepHeight;
-			bool bIsSupported = bIsWithinReach && !bIsLiftingOffSurface;
-			float TimeSinceSupported = MaxUnsupportedTimeBeforeFalling;
-			Blackboard->TryGet(CommonBlackboard::TimeSinceSupported, TimeSinceSupported);
-
-			if (bIsSupported)
-			{
-				Blackboard->Set(CommonBlackboard::TimeSinceSupported, 0.0f);
-			}
-			else if (!bIsLiftingOffSurface)
-			{
-				// Falling
-				TimeSinceSupported += DeltaSeconds;
-				Blackboard->Set(CommonBlackboard::TimeSinceSupported, TimeSinceSupported);
-				bIsSupported = TimeSinceSupported < MaxUnsupportedTimeBeforeFalling;
-			}
-			else
-			{
-				// Moving up relative to ground
-				Blackboard->Set(CommonBlackboard::TimeSinceSupported, MaxUnsupportedTimeBeforeFalling);
-			}
-
-			if (!bIsSupported)
+			if (const IChaosCharacterMovementModeInterface* Mode = Cast<IChaosCharacterMovementModeInterface>(BaseMode.Get()))
 			{
-				EvalResult.NextMode = TransitionToFallingMode;
+				const float DeltaSeconds = Params.TimeStep.StepMs * 0.001f;
+				const FVector LinearVelocity = StartingSyncState->GetVelocity_WorldSpace();
+
+				const FVector ProjectedVelocity = StartingSyncState->GetVelocity_WorldSpace() + SimInputs->Gravity * DeltaSeconds;
+				FVector ProjectedGroundVelocity = UChaosGroundMovementUtils::ComputeLocalGroundVelocity_Internal(StartingSyncState->GetLocation_WorldSpace(), FloorResult);
+				const Chaos::FPBDRigidParticleHandle* GroundParticle = UChaosGroundMovementUtils::GetRigidParticleHandleFromFloorResult_Internal(FloorResult);
+				if (GroundParticle && GroundParticle->IsDynamic() && GroundParticle->GravityEnabled())
+				{
+					// This might not be correct if different physics objects have different gravity but is saves having to go
+					// to the component to get the gravity on the physics volume.
+					ProjectedGroundVelocity += SimInputs->PhysicsObjectGravity * SimInputs->UpDir * DeltaSeconds;
+				}
+				const bool bIsGroundMoving = ProjectedGroundVelocity.SizeSquared() > UE_KINDA_SMALL_NUMBER;
+				const FVector ProjectedRelativeVelocity = ProjectedVelocity - ProjectedGroundVelocity;
+				const float ProjectedRelativeNormalVelocity = FloorResult.HitResult.ImpactNormal.Dot(ProjectedVelocity - ProjectedGroundVelocity);
+				const float ProjectedRelativeVerticalVelocity = SimInputs->UpDir.Dot(ProjectedVelocity - ProjectedGroundVelocity);
+				const float VerticalVelocityLimit = 2.0f / DeltaSeconds;
+				bool bIsLiftingOffSurface = false;
+				if ((ProjectedRelativeNormalVelocity > VerticalVelocityLimit) && bIsGroundMoving && (ProjectedRelativeVerticalVelocity > VerticalVelocityLimit))
+				{
+					bIsLiftingOffSurface = true;
+				}
+
+				const bool bIsWithinReach = FloorResult.FloorDist - Mode->GetTargetHeight() <= SharedSettings->MaxStepHeight;
+				bool bIsSupported = bIsWithinReach && !bIsLiftingOffSurface;
+				float TimeSinceSupported = MaxUnsupportedTimeBeforeFalling;
+				Blackboard->TryGet(CommonBlackboard::TimeSinceSupported, TimeSinceSupported);
+
+				if (bIsSupported)
+				{
+					Blackboard->Set(CommonBlackboard::TimeSinceSupported, 0.0f);
+				}
+				else if (!bIsLiftingOffSurface)
+				{
+					// Falling
+					TimeSinceSupported += DeltaSeconds;
+					Blackboard->Set(CommonBlackboard::TimeSinceSupported, TimeSinceSupported);
+					bIsSupported = TimeSinceSupported < MaxUnsupportedTimeBeforeFalling;
+				}
+				else
+				{
+					// Moving up relative to ground
+					Blackboard->Set(CommonBlackboard::TimeSinceSupported, MaxUnsupportedTimeBeforeFalling);
+				}
+
+				if (!bIsSupported)
+				{
+					EvalResult.NextMode = TransitionToFallingMode;
+				}
 			}
 		}
 	}
diff --git a/Engine/Plugins/Experimental/ChaosMover/Source/ChaosMover/Private/Character/Transitions/ChaosCharacterLandingCheck.cpp b/Engine/Plugins/Experimental/ChaosMover/Source/ChaosMover/Private/Character/Transitions/ChaosCharacterLandingCheck.cpp
--- a/Engine/Plugins/Experimental/ChaosMover/Source/ChaosMover/Private/Character/Transitions/ChaosCharacterLandingCheck.cpp
+++ b/Engine/Plugins/Experimental/ChaosMover/Source/ChaosMover/Private/Character/Transitions/ChaosCharacterLandingCheck.cpp
@@ -41,40 +41,43 @@ FTransitionEvalResult UChaosCharacterLandingCheck::Evaluate_Implementation(const
 			return EvalResult;
 		}
 
-		if (const IChaosCharacterMovementModeInterface* Mode = Cast<IChaosCharacterMovementModeInterface>(Simulation->FindMovementModeByName(Params.StartState.SyncState.MovementMode)))
+		if (TStrongObjectPtr<const UBaseMovementMode> BaseMode = Simulation->FindMovementModeByName(Params.StartState.SyncState.MovementMode).Pin())
 		{
-			const float DeltaSeconds = Params.TimeStep.StepMs * 0.001f;
-			const FVector LinearVelocity = StartingSyncState->GetVelocity_WorldSpace();
-			const float VerticalVelocity = SimInputs->UpDir.Dot(LinearVelocity);
-			const bool bJumping = CharacterInputs && CharacterInputs->bIsJumpJustPressed;
-			const bool bIsMovingUp = bJumping || (VerticalVelocity > 0.0f);
-
-			// Check for ground landing
-			if (!TransitionToGroundMode.IsNone())
+			if (const IChaosCharacterMovementModeInterface* Mode = Cast<IChaosCharacterMovementModeInterface>(BaseMode.Get()))
 			{
-				FFloorCheckResult FloorResult;
-				if (Blackboard->TryGet(CommonBlackboard::LastFloorResult, FloorResult) && FloorResult.IsWalkableFloor())
+				const float DeltaSeconds = Params.TimeStep.StepMs * 0.001f;
+				const FVector LinearVelocity = StartingSyncState->GetVelocity_WorldSpace();
+				const float VerticalVelocity = SimInputs->UpDir.Dot(LinearVelocity);
+				const bool bJumping = CharacterInputs && CharacterInputs->bIsJumpJustPressed;
+				const bool bIsMovingUp = bJumping || (VerticalVelocity > 0.0f);
+
+				// Check for ground landing
+				if (!TransitionToGroundMode.IsNone())
 				{
-					const float TargetHeight = Mode->GetTargetHeight();
-					const float FloorDistanceWithFloorNormal = FloorResult.HitResult.ImpactNormal.Dot(StartingSyncState->GetLocation_WorldSpace() - FloorResult.HitResult.ImpactPoint);
-					const FVector LocalGroundVelocity = UChaosGroundMovementUtils::ComputeLocalGroundVelocity_Internal(StartingSyncState->GetLocation_WorldSpace(), FloorResult);
-					const float RelativeVerticalVelocity = FloorResult.HitResult.ImpactNormal.Dot(LinearVelocity - LocalGroundVelocity);
-					const float ProjectedFloorDistance = FloorDistanceWithFloorNormal + RelativeVerticalVelocity * DeltaSeconds;
-					const bool bIsFloorWithinReach = ProjectedFloorDistance < TargetHeight + FloorDistanceTolerance + UE_KINDA_SMALL_NUMBER;
-					const bool bIsMovingUpRelativeToFloor = (RelativeVerticalVelocity > UE_KINDA_SMALL_NUMBER) || bJumping;
-
-					if (bIsFloorWithinReach && !bIsMovingUpRelativeToFloor)
+					FFloorCheckResult FloorResult;
+					if (Blackboard->TryGet(CommonBlackboard::LastFloorResult, FloorResult) && FloorResult.IsWalkableFloor())
 					{
-						if (Simulation->FindMovementModeByName(TransitionToGroundMode))
+						const float TargetHeight = Mode->GetTargetHeight();
+						const float FloorDistanceWithFloorNormal = FloorResult.HitResult.ImpactNormal.Dot(StartingSyncState->GetLocation_WorldSpace() - FloorResult.HitResult.ImpactPoint);
+						const FVector LocalGroundVelocity = UChaosGroundMovementUtils::ComputeLocalGroundVelocity_Internal(StartingSyncState->GetLocation_WorldSpace(), FloorResult);
+						const float RelativeVerticalVelocity = FloorResult.HitResult.ImpactNormal.Dot(LinearVelocity - LocalGroundVelocity);
+						const float ProjectedFloorDistance = FloorDistanceWithFloorNormal + RelativeVerticalVelocity * DeltaSeconds;
+						const bool bIsFloorWithinReach = ProjectedFloorDistance < TargetHeight + FloorDistanceTolerance + UE_KINDA_SMALL_NUMBER;
+						const bool bIsMovingUpRelativeToFloor = (RelativeVerticalVelocity > UE_KINDA_SMALL_NUMBER) || bJumping;
+
+						if (bIsFloorWithinReach && !bIsMovingUpRelativeToFloor)
 						{
-							EvalResult.NextMode = TransitionToGroundMode;
-						}
-						else
-						{
-							UE_LOG(LogChaosMover, Warning, TEXT("Invalid ground mode name %s in ChaosCharacterLandingCheck. Cannot make transition"), *TransitionToGroundMode.ToString());
-						}
+							if (Simulation->HasMovementMode(TransitionToGroundMode))
+							{
+								EvalResult.NextMode = TransitionToGroundMode;
+							}
+							else
+							{
+								UE_LOG(LogChaosMover, Warning, TEXT("Invalid ground mode name %s in ChaosCharacterLandingCheck. Cannot make transition"), *TransitionToGroundMode.ToString());
+							}
 						
-						return EvalResult;
+							return EvalResult;
+						}
 					}
 				}
 			}
diff --git a/Engine/Plugins/Experimental/ChaosMover/Source/ChaosMover/Private/Character/Transitions/ChaosCharacterWaterCheck.cpp b/Engine/Plugins/Experimental/ChaosMover/Source/ChaosMover/Private/Character/Transitions/ChaosCharacterWaterCheck.cpp
--- a/Engine/Plugins/Experimental/ChaosMover/Source/ChaosMover/Private/Character/Transitions/ChaosCharacterWaterCheck.cpp
+++ b/Engine/Plugins/Experimental/ChaosMover/Source/ChaosMover/Private/Character/Transitions/ChaosCharacterWaterCheck.cpp
@@ -41,73 +41,76 @@ FTransitionEvalResult UChaosCharacterWaterCheck::Evaluate_Implementation(const F
 			return EvalResult;
 		}
 
-		if (const IChaosCharacterMovementModeInterface* Mode = Cast<IChaosCharacterMovementModeInterface>(Simulation->FindMovementModeByName(Params.StartState.SyncState.MovementMode)))
+		if (TStrongObjectPtr<const UBaseMovementMode> BaseMode = Simulation->FindMovementModeByName(Params.StartState.SyncState.MovementMode).Pin())
 		{
-			const float DeltaSeconds = Params.TimeStep.StepMs * 0.001f;
-			const FVector LinearVelocity = StartingSyncState->GetVelocity_WorldSpace();
-			const float VerticalVelocity = SimInputs->UpDir.Dot(LinearVelocity);
-			const bool bJumping = CharacterInputs && CharacterInputs->bIsJumpJustPressed;
-			const bool bIsMovingUp = bJumping || (VerticalVelocity > UE_KINDA_SMALL_NUMBER);
+			if (const IChaosCharacterMovementModeInterface* Mode = Cast<IChaosCharacterMovementModeInterface>(BaseMode.Get()))
+			{
+				const float DeltaSeconds = Params.TimeStep.StepMs * 0.001f;
+				const FVector LinearVelocity = StartingSyncState->GetVelocity_WorldSpace();
+				const float VerticalVelocity = SimInputs->UpDir.Dot(LinearVelocity);
+				const bool bJumping = CharacterInputs && CharacterInputs->bIsJumpJustPressed;
+				const bool bIsMovingUp = bJumping || (VerticalVelocity > UE_KINDA_SMALL_NUMBER);
 
-			FFindFloorResult FloorResult;
-			FWaterCheckResult WaterResult;
+				FFindFloorResult FloorResult;
+				FWaterCheckResult WaterResult;
 
-			Blackboard->TryGet(CommonBlackboard::LastFloorResult, FloorResult);
-			Blackboard->TryGet(CommonBlackboard::LastWaterResult, WaterResult);
+				Blackboard->TryGet(CommonBlackboard::LastFloorResult, FloorResult);
+				Blackboard->TryGet(CommonBlackboard::LastWaterResult, WaterResult);
 			
-			const float ProjectedImmersionDepth = WaterResult.WaterSplineData.ImmersionDepth - VerticalVelocity * DeltaSeconds;
-			const bool bStartSwimming = ProjectedImmersionDepth > WaterModeStartImmersionDepth;
-			const bool bInWater = StartState.SyncState.MovementMode == WaterModeName;
-			const bool bStopSwimming = bInWater && (ProjectedImmersionDepth < WaterModeStopImmersionDepth);
+				const float ProjectedImmersionDepth = WaterResult.WaterSplineData.ImmersionDepth - VerticalVelocity * DeltaSeconds;
+				const bool bStartSwimming = ProjectedImmersionDepth > WaterModeStartImmersionDepth;
+				const bool bInWater = StartState.SyncState.MovementMode == WaterModeName;
+				const bool bStopSwimming = bInWater && (ProjectedImmersionDepth < WaterModeStopImmersionDepth);
 
-			// Check if we need to enter/exit water mode
-			if (WaterResult.IsSwimmableVolume() && bStartSwimming)
-			{
-				if (!WaterModeName.IsNone())
-				{
-					if (Simulation->FindMovementModeByName(WaterModeName))
-					{
-						EvalResult.NextMode = WaterModeName;
-					}
-					else
-					{
-						UE_LOG(LogChaosMover, Warning, TEXT("Invalid water mode name %s in ChaosCharacterWaterCheck. Cannot make transition"), *WaterModeName.ToString());
-					}
-					return EvalResult;
-				}
-			}
-			else if (bStopSwimming)
-			{
-				const bool bIsWithinReach = FloorResult.FloorDist <= Mode->GetTargetHeight();
- 
-				if (FloorResult.IsWalkableFloor() && bIsWithinReach && !bIsMovingUp)
+				// Check if we need to enter/exit water mode
+				if (WaterResult.IsSwimmableVolume() && bStartSwimming)
 				{
-					if (!GroundModeName.IsNone())
+					if (!WaterModeName.IsNone())
 					{
-						if (Simulation->FindMovementModeByName(GroundModeName))
+						if (Simulation->HasMovementMode(WaterModeName))
 						{
-							EvalResult.NextMode = GroundModeName;
+							EvalResult.NextMode = WaterModeName;
 						}
 						else
 						{
-							UE_LOG(LogChaosMover, Warning, TEXT("Invalid ground mode name %s in ChaosCharacterWaterCheck. Cannot make transition"), *GroundModeName.ToString());
+							UE_LOG(LogChaosMover, Warning, TEXT("Invalid water mode name %s in ChaosCharacterWaterCheck. Cannot make transition"), *WaterModeName.ToString());
 						}
 						return EvalResult;
 					}
 				}
-				else
+				else if (bStopSwimming)
 				{
-					if (!AirModeName.IsNone())
+					const bool bIsWithinReach = FloorResult.FloorDist <= Mode->GetTargetHeight();
+ 
+					if (FloorResult.IsWalkableFloor() && bIsWithinReach && !bIsMovingUp)
 					{
-						if (Simulation->FindMovementModeByName(AirModeName))
+						if (!GroundModeName.IsNone())
 						{
-							EvalResult.NextMode = AirModeName;
+							if (Simulation->HasMovementMode(GroundModeName))
+							{
+								EvalResult.NextMode = GroundModeName;
+							}
+							else
+							{
+								UE_LOG(LogChaosMover, Warning, TEXT("Invalid ground mode name %s in ChaosCharacterWaterCheck. Cannot make transition"), *GroundModeName.ToString());
+							}
+							return EvalResult;
 						}
-						else
+					}
+					else
+					{
+						if (!AirModeName.IsNone())
 						{
-							UE_LOG(LogChaosMover, Warning, TEXT("Invalid air mode name %s in ChaosCharacterWaterCheck. Cannot make transition"), *AirModeName.ToString());
+							if (Simulation->HasMovementMode(AirModeName))
+							{
+								EvalResult.NextMode = AirModeName;
+							}
+							else
+							{
+								UE_LOG(LogChaosMover, Warning, TEXT("Invalid air mode name %s in ChaosCharacterWaterCheck. Cannot make transition"), *AirModeName.ToString());
+							}
+							return EvalResult;
 						}
-						return EvalResult;
 					}
 				}
 			}
diff --git a/Engine/Plugins/Experimental/ChaosMover/Source/ChaosMover/Public/ChaosMover/ChaosMoverSimulation.h b/Engine/Plugins/Experimental/ChaosMover/Source/ChaosMover/Public/ChaosMover/ChaosMoverSimulation.h
--- a/Engine/Plugins/Experimental/ChaosMover/Source/ChaosMover/Public/ChaosMover/ChaosMoverSimulation.h
+++ b/Engine/Plugins/Experimental/ChaosMover/Source/ChaosMover/Public/ChaosMover/ChaosMoverSimulation.h
@@ -39,14 +39,10 @@ public:
 	UFUNCTION(BlueprintPure, Category = Mover)
 	CHAOSMOVER_API FMoverDataCollection& GetLocalSimInput_Mutable();
 
-	UFUNCTION(BlueprintPure, Category = Mover)
-	CHAOSMOVER_API const UBaseMovementMode* GetCurrentMovementMode() const;
-
-	UFUNCTION(BlueprintPure, Category = Mover)
-	CHAOSMOVER_API const UBaseMovementMode* FindMovementModeByName(const FName& Name) const;
-
-	UFUNCTION(BlueprintPure, Category = Mover)
-	CHAOSMOVER_API UBaseMovementMode* FindMovementModeByName_Mutable(const FName& Name);
+	CHAOSMOVER_API TWeakObjectPtr<const UBaseMovementMode> GetCurrentMovementMode() const;
+	CHAOSMOVER_API TWeakObjectPtr<const UBaseMovementMode> FindMovementModeByName(const FName& Name) const;
+	CHAOSMOVER_API TWeakObjectPtr<UBaseMovementMode> FindMovementModeByName_Mutable(const FName& Name);
+	CHAOSMOVER_API bool HasMovementMode(const FName& Name) const;
 
 	UFUNCTION(BlueprintCallable, CustomThunk, Category = Mover, meta = (CustomStructureParam = "EffectAsRawData", AllowAbstract = "false", DisplayName = "Queue Instant Movement Effect"))
 	void K2_QueueInstantMovementEffect(UPARAM(DisplayName = "Instant Movement Effect") const int32& EffectAsRawData);
