ChaosMover: Adding scopes to better understand where perf is going

#Synced-CL 48488436

diff --git a/Engine/Plugins/Experimental/ChaosMover/Source/ChaosMover/Private/ChaosMoverStateMachine.cpp b/Engine/Plugins/Experimental/ChaosMover/Source/ChaosMover/Private/ChaosMoverStateMachine.cpp
--- a/Engine/Plugins/Experimental/ChaosMover/Source/ChaosMover/Private/ChaosMoverStateMachine.cpp
+++ b/Engine/Plugins/Experimental/ChaosMover/Source/ChaosMover/Private/ChaosMoverStateMachine.cpp
@@ -278,6 +278,8 @@ namespace UE::ChaosMover
 
 	void FMoverStateMachine::ReceiveInstantMovementEffects(const FMoverTimeStep& TimeStep, const FChaosNetInstantMovementEffectsQueue* InstantMovementEffectsQueue)
 	{
+		QUICK_SCOPE_CYCLE_COUNTER(MoverStateMachine_ReceiveInstantMovementEffects);
+
 		InstantMovementEffectsIDHistory.CullOldFrames(TimeStep.ServerFrame, CVars::InstantMovementEffectIDHistorySize);
 
 		// Copy queued instant movement effects from the input to the state machine
@@ -318,6 +320,8 @@ namespace UE::ChaosMover
 
 	void FMoverStateMachine::OnSimulationTick(const FMoverTimeStep& TimeStep, const FMoverTickStartData& StartState, UMoverBlackboard* SimBlackboard, UMovementMixer* MovementMixer, FMoverTickEndData& OutputState)
 	{
+		QUICK_SCOPE_CYCLE_COUNTER(MoverStateMachine_OnSimulationTick);
+
 		FMoverTimeStep SubTimeStep = TimeStep;
 		WorkingSubstepStartData = StartState;
 		bool bIsWorkingStartStateReady = true;	// this flag is used to avoid unneeded data copying after substeps
@@ -459,7 +463,7 @@ namespace UE::ChaosMover
 				if (!CVars::bSkipGenerateMoveIfOverridden ||
 					!(bHasLayeredMoveContributions && CombinedLayeredMove.MixMode == EMoveMixMode::OverrideAll))
 				{
-					QUICK_SCOPE_CYCLE_COUNTER(STAT_GenerateMoveFromMode);
+					QUICK_SCOPE_CYCLE_COUNTER(MoverStateMachine_OnSimulationTick_GenerateMoveFromMode);
 					CurrentMode->GenerateMove(WorkingSubstepStartData, SubTimeStep, OUT CombinedMove);
 
 					bHasModeMoveContribution = true;
@@ -552,6 +556,7 @@ namespace UE::ChaosMover
 					}
 					else
 					{
+						QUICK_SCOPE_CYCLE_COUNTER(MoverStateMachine_OnSimulationTick_CurrentModeSimulationTick);
 						CurrentMode->SimulationTick(WorkingSimTickParams, OutputState);
 					}
 
@@ -638,6 +643,8 @@ namespace UE::ChaosMover
 
 	void FMoverStateMachine::OnSimulationRollback(const FMoverTimeStep& NewTimeStep, const FMoverSyncState& InvalidSyncState, const FMoverSyncState& NewSyncState)
 	{
+		QUICK_SCOPE_CYCLE_COUNTER(MoverStateMachine_OnSimulationRollback);
+
 		const int32 PreviousInternalServerFrame = InternalSimTime.ServerFrame;
 		InternalSimTime = NewTimeStep;
 
@@ -781,6 +788,8 @@ namespace UE::ChaosMover
 
 	void FMoverStateMachine::AdvanceToNextMode()
 	{
+		QUICK_SCOPE_CYCLE_COUNTER(MoverStateMachine_AdvanceToNextMode);
+
 		TStrongObjectPtr<UImmediateMovementModeTransition> QueuedModeTransition = QueuedModeTransitionWeakPtr.Pin();
 		if (QueuedModeTransition)
 		{
@@ -850,6 +859,8 @@ namespace UE::ChaosMover
 
 	bool FMoverStateMachine::ApplyInstantEffects(const FMoverTickStartData& SubstepStartData, const FMoverTimeStep& TimeStep, FMoverSyncState& OutputState)
 	{
+		QUICK_SCOPE_CYCLE_COUNTER(MoverStateMachine_ApplyInstantEffects);
+
 		FApplyMovementEffectParams_Async EffectParams_Async;
 		EffectParams_Async.StartState = &SubstepStartData;
 		EffectParams_Async.TimeStep = &TimeStep;
