ChaosMover: Adding scopes to know where expensive traces are coming from

#Synced-CL 48488436

diff --git a/Engine/Plugins/Experimental/ChaosMover/Source/ChaosMover/Private/Utilities/ChaosMoverQueryUtils.cpp b/Engine/Plugins/Experimental/ChaosMover/Source/ChaosMover/Private/Utilities/ChaosMoverQueryUtils.cpp
--- a/Engine/Plugins/Experimental/ChaosMover/Source/ChaosMover/Private/Utilities/ChaosMoverQueryUtils.cpp
+++ b/Engine/Plugins/Experimental/ChaosMover/Source/ChaosMover/Private/Utilities/ChaosMoverQueryUtils.cpp
@@ -89,6 +89,8 @@ namespace UE::ChaosMover::Utils
 
 	void FloorSweep_Internal(const FFloorSweepParams& Params, FFloorCheckResult& OutFloorResult, FWaterCheckResult& OutWaterResult)
 	{
+		QUICK_SCOPE_CYCLE_COUNTER(ChaosMoverUtils_FloorSweep);
+
 		Chaos::EnsureIsInPhysicsThreadContext();
 
 		TArray<FHitResult> Hits;
@@ -162,6 +164,8 @@ namespace UE::ChaosMover::Utils
 
 	bool CapsuleOverlapTest_Internal(const FCapsuleOverlapTestParams& Params)
 	{
+		QUICK_SCOPE_CYCLE_COUNTER(ChaosMoverUtils_CapsuleOverlapTest);
+
 		Chaos::EnsureIsInPhysicsThreadContext();
 
 		TArray<FHitResult> Hits;
