Update Velocity Padding Calculation to fix clipping through the floor.

diff --git a/Engine/Plugins/Experimental/ChaosMover/Source/ChaosMover/Private/Character/Modes/ChaosWalkingMode.cpp b/Engine/Plugins/Experimental/ChaosMover/Source/ChaosMover/Private/Character/Modes/ChaosWalkingMode.cpp
--- a/Engine/Plugins/Experimental/ChaosMover/Source/ChaosMover/Private/Character/Modes/ChaosWalkingMode.cpp
+++ b/Engine/Plugins/Experimental/ChaosMover/Source/ChaosMover/Private/Character/Modes/ChaosWalkingMode.cpp
@@ -25,6 +25,10 @@
 
 #include UE_INLINE_GENERATED_CPP_BY_NAME(ChaosWalkingMode)
 
+bool bVelocityPaddingWalkingModeRollbackUseNewCode = true;
+FAutoConsoleVariableRef CVarVelocityPaddingWalkingModeRollbackUseNewCode(TEXT("ChaosMover.WalkingMode.VelocityPaddingUseNewCode"),
+	bVelocityPaddingWalkingModeRollbackUseNewCode, TEXT("If true use the new code path for calculating the VelocityPadding in ChaosWalkingMode\n"));
+
 UChaosWalkingMode::UChaosWalkingMode(const FObjectInitializer& ObjectInitializer)
 	: Super(ObjectInitializer)
 {
@@ -340,7 +344,16 @@ void UChaosWalkingMode::GetFloorAndCheckMovement(const FMoverDefaultSyncState& S
 	FVector DeltaPos = ProposedMove.LinearVelocity * DeltaSeconds;
 	OutDeltaPos = DeltaPos;
 
-	const float VelocityPadding = FMath::Max(DefaultSimInputs.UpDir.Dot(SyncState.GetVelocity_WorldSpace()) * DeltaSeconds, 0.0f);
+	float VelocityPadding = 0.0f;
+	if (bVelocityPaddingWalkingModeRollbackUseNewCode)
+	{
+		const FVector VelocityWithGravity = SyncState.GetVelocity_WorldSpace() + UMovementUtils::ComputeVelocityFromGravity(DefaultSimInputs.Gravity, DeltaSeconds);
+		VelocityPadding = FMath::Max(DefaultSimInputs.UpDir.Dot(-VelocityWithGravity) * DeltaSeconds, 0.0f);
+	}
+	else
+	{
+		VelocityPadding = FMath::Max(DefaultSimInputs.UpDir.Dot(SyncState.GetVelocity_WorldSpace()) * DeltaSeconds, 0.0f);
+	}
 
 	UE::ChaosMover::Utils::FFloorSweepParams SweepParams{
 		.ResponseParams = DefaultSimInputs.CollisionResponseParams,
