[Chaos Mover] Respect first subframe only flag on global transitions

#Synced-CL 48604082

diff --git a/Engine/Plugins/Experimental/ChaosMover/Source/ChaosMover/Private/ChaosMoverStateMachine.cpp b/Engine/Plugins/Experimental/ChaosMover/Source/ChaosMover/Private/ChaosMoverStateMachine.cpp
--- a/Engine/Plugins/Experimental/ChaosMover/Source/ChaosMover/Private/ChaosMoverStateMachine.cpp
+++ b/Engine/Plugins/Experimental/ChaosMover/Source/ChaosMover/Private/ChaosMoverStateMachine.cpp
@@ -532,7 +532,7 @@ namespace UE::ChaosMover
 						for (TWeakObjectPtr<UBaseMovementModeTransition> TransitionWeakPtr : GlobalTransitions)
 						{
 							TStrongObjectPtr<UBaseMovementModeTransition> Transition = TransitionWeakPtr.Pin();
-							if (Transition)
+							if (Transition && ((SubStepCount == 0) || !Transition->bFirstSubStepOnly))
 							{
 								EvalResult = Transition->Evaluate(WorkingSimTickParams);
 
