ChaosMover: Adding performance scopes

#Synced-CL 48488436

diff --git a/Engine/Plugins/Experimental/ChaosMover/Source/ChaosMover/Private/Character/Modes/ChaosFallingMode.cpp b/Engine/Plugins/Experimental/ChaosMover/Source/ChaosMover/Private/Character/Modes/ChaosFallingMode.cpp
--- a/Engine/Plugins/Experimental/ChaosMover/Source/ChaosMover/Private/Character/Modes/ChaosFallingMode.cpp
+++ b/Engine/Plugins/Experimental/ChaosMover/Source/ChaosMover/Private/Character/Modes/ChaosFallingMode.cpp
@@ -44,6 +44,8 @@ UChaosFallingMode::UChaosFallingMode(const FObjectInitializer& ObjectInitializer
 
 void UChaosFallingMode::GenerateMove_Implementation(const FMoverTickStartData& StartState, const FMoverTimeStep& TimeStep, FProposedMove& OutProposedMove) const
 {
+	QUICK_SCOPE_CYCLE_COUNTER(ChaosFallingMode_GenerateMove);
+
 	if (!Simulation)
 	{
 		UE_LOG(LogChaosMover, Warning, TEXT("No Simulation set on ChaosFallingMode"));
@@ -175,6 +177,8 @@ void UChaosFallingMode::GenerateMove_Implementation(const FMoverTickStartData& S
 
 void UChaosFallingMode::SimulationTick_Implementation(const FSimulationTickParams& Params, FMoverTickEndData& OutputState)
 {
+	QUICK_SCOPE_CYCLE_COUNTER(ChaosFallingMode_SimulationTick);
+
 	if (!Simulation)
 	{
 		UE_LOG(LogChaosMover, Warning, TEXT("No Simulation set on ChaosFallingMode"));
diff --git a/Engine/Plugins/Experimental/ChaosMover/Source/ChaosMover/Private/Character/Modes/ChaosFlyingMode.cpp b/Engine/Plugins/Experimental/ChaosMover/Source/ChaosMover/Private/Character/Modes/ChaosFlyingMode.cpp
--- a/Engine/Plugins/Experimental/ChaosMover/Source/ChaosMover/Private/Character/Modes/ChaosFlyingMode.cpp
+++ b/Engine/Plugins/Experimental/ChaosMover/Source/ChaosMover/Private/Character/Modes/ChaosFlyingMode.cpp
@@ -26,6 +26,8 @@ UChaosFlyingMode::UChaosFlyingMode(const FObjectInitializer& ObjectInitializer)
 
 void UChaosFlyingMode::GenerateMove_Implementation(const FMoverTickStartData& StartState, const FMoverTimeStep& TimeStep, FProposedMove& OutProposedMove) const
 {
+	QUICK_SCOPE_CYCLE_COUNTER(ChaosFlyingMode_GenerateMove);
+
 	if (!Simulation)
 	{
 		UE_LOG(LogChaosMover, Warning, TEXT("No Simulation set on ChaosFlyingMode"));
@@ -99,6 +101,8 @@ void UChaosFlyingMode::GenerateMove_Implementation(const FMoverTickStartData& St
 
 void UChaosFlyingMode::SimulationTick_Implementation(const FSimulationTickParams& Params, FMoverTickEndData& OutputState)
 {
+	QUICK_SCOPE_CYCLE_COUNTER(ChaosFlyingMode_SimulationTick);
+
 	if (!Simulation)
 	{
 		UE_LOG(LogChaosMover, Warning, TEXT("No Simulation set on ChaosFlyingMode"));
diff --git a/Engine/Plugins/Experimental/ChaosMover/Source/ChaosMover/Private/Character/Modes/ChaosSwimmingMode.cpp b/Engine/Plugins/Experimental/ChaosMover/Source/ChaosMover/Private/Character/Modes/ChaosSwimmingMode.cpp
--- a/Engine/Plugins/Experimental/ChaosMover/Source/ChaosMover/Private/Character/Modes/ChaosSwimmingMode.cpp
+++ b/Engine/Plugins/Experimental/ChaosMover/Source/ChaosMover/Private/Character/Modes/ChaosSwimmingMode.cpp
@@ -29,25 +29,29 @@ UChaosSwimmingMode::UChaosSwimmingMode(const FObjectInitializer& ObjectInitializ
 
 void UChaosSwimmingMode::GenerateMove_Implementation(const FMoverTickStartData& StartState, const FMoverTimeStep& TimeStep, FProposedMove& OutProposedMove) const
 {
-	if (!Simulation)
+	QUICK_SCOPE_CYCLE_COUNTER(ChaosSwimmingMode_GenerateMove);
+
+	if (!ensureMsgf(Simulation, TEXT("No Simulation set on ChaosSwimmingMode")))
 	{
-		UE_LOG(LogChaosMover, Warning, TEXT("No Simulation set on ChaosSwimmingMode"));
 		return;
 	}
-
 	const FChaosMoverSimulationDefaultInputs* DefaultSimInputs = Simulation->GetLocalSimInput().FindDataByType<FChaosMoverSimulationDefaultInputs>();
-	const FCharacterDefaultInputs* CharacterInputs = StartState.InputCmd.InputCollection.FindDataByType<FCharacterDefaultInputs>();
-	if (!DefaultSimInputs || !CharacterInputs)
+	if (!ensureMsgf(DefaultSimInputs, TEXT("ChaosSwimmingMode requires FChaosMoverSimulationDefaultInputs")))
 	{
-		UE_LOG(LogChaosMover, Warning, TEXT("ChaosSwimmingMode requires FChaosMoverSimulationDefaultInputs and FCharacterDefaultInputs"));
 		return;
 	}
-	
+	const FCharacterDefaultInputs* CharacterInputs = StartState.InputCmd.InputCollection.FindDataByType<FCharacterDefaultInputs>();
+	if (!ensureMsgf(CharacterInputs, TEXT("ChaosSwimmingMode requires FCharacterDefaultInputs")))
+	{
+		return;
+	}	
 	const FMoverDefaultSyncState* StartingSyncState = StartState.SyncState.SyncStateCollection.FindDataByType<FMoverDefaultSyncState>();
-	check(StartingSyncState);
-	
-	const float DeltaSeconds = TimeStep.StepMs * 0.001f;
+	if (!ensureMsgf(StartingSyncState, TEXT("ChaosSwimmingMode requires FMoverDefaultSyncState")))
+	{
+		return;
+	}
 	
+	const float DeltaSeconds = TimeStep.StepMs * 0.001f;	
 	const float HalfHeight = DefaultSimInputs->PawnCollisionHalfHeight;
 	FVector Velocity = StartingSyncState->GetVelocity_WorldSpace();
 	
@@ -215,21 +219,32 @@ void UChaosSwimmingMode::GenerateMove_Implementation(const FMoverTickStartData&
 
 void UChaosSwimmingMode::SimulationTick_Implementation(const FSimulationTickParams& Params, FMoverTickEndData& OutputState)
 {
-	if (!Simulation)
+	QUICK_SCOPE_CYCLE_COUNTER(ChaosSwimmingMode_SimulationTick);
+
+	if (!ensureMsgf(Simulation, TEXT("No Simulation set on ChaosSwimmingMode")))
 	{
-		UE_LOG(LogChaosMover, Warning, TEXT("No Simulation set on ChaosSwimmingMode"));
 		return;
-	}
-	
+	}	
 	const FChaosMoverSimulationDefaultInputs* DefaultSimInputs = Simulation->GetLocalSimInput().FindDataByType<FChaosMoverSimulationDefaultInputs>();
+	if (!ensureMsgf(DefaultSimInputs != nullptr, TEXT("No FChaosMoverSimulationDefaultInputs in UChaosSwimmingMode::SimulationTick_Implementation, skipping")))
+	{
+		return;
+	}
 	const FMoverTickStartData& StartState = Params.StartState;
 	const FCharacterDefaultInputs* CharacterInputs = StartState.InputCmd.InputCollection.FindDataByType<FCharacterDefaultInputs>();
+	if (!ensureMsgf(CharacterInputs, TEXT("ChaosSwimmingMode requires FCharacterDefaultInputs")))
+	{
+		return;
+	}
+	const FMoverDefaultSyncState* StartingSyncState = StartState.SyncState.SyncStateCollection.FindDataByType<FMoverDefaultSyncState>();
+	if (!ensureMsgf(StartingSyncState, TEXT("ChaosSwimmingMode requires FMoverDefaultSyncState")))
+	{
+		return;
+	}
+
 	FProposedMove ProposedMove = Params.ProposedMove;
 	const FVector UpDir = DefaultSimInputs->UpDir;
 
-	const FMoverDefaultSyncState* StartingSyncState = StartState.SyncState.SyncStateCollection.FindDataByType<FMoverDefaultSyncState>();
-	check(StartingSyncState);
-
 	FMoverDefaultSyncState& OutputSyncState = OutputState.SyncState.SyncStateCollection.FindOrAddMutableDataByType<FMoverDefaultSyncState>();
 	OutputSyncState = *StartingSyncState;
 
diff --git a/Engine/Plugins/Experimental/ChaosMover/Source/ChaosMover/Private/Character/Modes/ChaosWalkingMode.cpp b/Engine/Plugins/Experimental/ChaosMover/Source/ChaosMover/Private/Character/Modes/ChaosWalkingMode.cpp
--- a/Engine/Plugins/Experimental/ChaosMover/Source/ChaosMover/Private/Character/Modes/ChaosWalkingMode.cpp
+++ b/Engine/Plugins/Experimental/ChaosMover/Source/ChaosMover/Private/Character/Modes/ChaosWalkingMode.cpp
@@ -53,9 +53,11 @@ void UChaosWalkingMode::UpdateConstraintSettings(Chaos::FCharacterGroundConstrai
 
 void UChaosWalkingMode::GenerateMove_Implementation(const FMoverTickStartData& StartState, const FMoverTimeStep& TimeStep, FProposedMove& OutProposedMove) const
 {
+	QUICK_SCOPE_CYCLE_COUNTER(ChaosWalkingMode_GenerateMove);
+
 	if (!Simulation)
 	{
-		UE_LOG(LogChaosMover, Warning, TEXT("No Simulation set on ChaosFallingMode"));
+		UE_LOG(LogChaosMover, Warning, TEXT("No Simulation set on ChaosWalkingMode"));
 		return;
 	}
 
@@ -63,7 +65,7 @@ void UChaosWalkingMode::GenerateMove_Implementation(const FMoverTickStartData& S
 	const FCharacterDefaultInputs* CharacterInputs = StartState.InputCmd.InputCollection.FindDataByType<FCharacterDefaultInputs>();
 	if (!DefaultSimInputs || !CharacterInputs)
 	{
-		UE_LOG(LogChaosMover, Warning, TEXT("ChaosFallingMode requires FChaosMoverSimulationDefaultInputs and FCharacterDefaultInputs"));
+		UE_LOG(LogChaosMover, Warning, TEXT("ChaosWalkingMode requires FChaosMoverSimulationDefaultInputs and FCharacterDefaultInputs"));
 		return;
 	}
 
@@ -190,9 +192,11 @@ void UChaosWalkingMode::GenerateMove_Implementation(const FMoverTickStartData& S
 
 void UChaosWalkingMode::SimulationTick_Implementation(const FSimulationTickParams& Params, FMoverTickEndData& OutputState)
 {
+	QUICK_SCOPE_CYCLE_COUNTER(ChaosWalkingMode_SimulationTick);
+
 	if (!Simulation)
 	{
-		UE_LOG(LogChaosMover, Warning, TEXT("No Simulation set on ChaosFallingMode"));
+		UE_LOG(LogChaosMover, Warning, TEXT("No Simulation set on ChaosWalkingMode"));
 		return;
 	}
 
@@ -200,7 +204,7 @@ void UChaosWalkingMode::SimulationTick_Implementation(const FSimulationTickParam
 	const FCharacterDefaultInputs* CharacterInputs = Params.StartState.InputCmd.InputCollection.FindDataByType<FCharacterDefaultInputs>();
 	if (!DefaultSimInputs || !CharacterInputs)
 	{
-		UE_LOG(LogChaosMover, Warning, TEXT("ChaosFallingMode requires FChaosMoverSimulationDefaultInputs and FCharacterDefaultInputs"));
+		UE_LOG(LogChaosMover, Warning, TEXT("ChaosWalkingMode requires FChaosMoverSimulationDefaultInputs and FCharacterDefaultInputs"));
 		return;
 	}
 
diff --git a/Engine/Plugins/Experimental/ChaosMover/Source/ChaosMover/Private/PathedMovement/ChaosPathedMovementMode.cpp b/Engine/Plugins/Experimental/ChaosMover/Source/ChaosMover/Private/PathedMovement/ChaosPathedMovementMode.cpp
--- a/Engine/Plugins/Experimental/ChaosMover/Source/ChaosMover/Private/PathedMovement/ChaosPathedMovementMode.cpp
+++ b/Engine/Plugins/Experimental/ChaosMover/Source/ChaosMover/Private/PathedMovement/ChaosPathedMovementMode.cpp
@@ -57,8 +57,15 @@ void UChaosPathedMovementMode::GenerateMove_Implementation(const FMoverTickStart
 
 void UChaosPathedMovementMode::SimulationTick_Implementation(const FSimulationTickParams& Params, FMoverTickEndData& OutputState)
 {
+	QUICK_SCOPE_CYCLE_COUNTER(ChaosPathedMovementMode_SimulationTick);
+
 	Chaos::EnsureIsInPhysicsThreadContext();
 
+	if (!ensureMsgf(Simulation, TEXT("No Simulation set on ChaosPathedMovementMode")))
+	{
+		return;
+	}
+
 	// We retrieve the previous properties so we can compare with the new. If there are none we default to FChaosPathedMovementState()
 	const FChaosPathedMovementState* FoundPreviousMoveState = Params.StartState.SyncState.SyncStateCollection.FindDataByType<FChaosPathedMovementState>();
 	const FChaosPathedMovementState PreviousMoveState = FoundPreviousMoveState ? *FoundPreviousMoveState : FChaosPathedMovementState();	
@@ -152,7 +159,7 @@ void UChaosPathedMovementMode::SimulationTick_Implementation(const FSimulationTi
 	if (CurrentProperties.bWantsToPlay && (!PreviousMoveState.bHasFinished || bReceivedNewStartRequest))
 	{
 		const float FrameDuration = Params.TimeStep.StepMs * 0.001f;
-		const float PlaybackSpeed = (OneWayPlaybackDuration >= 0.0f) ? (1.0f / OneWayPlaybackDuration) : 0.0f;
+		const float PlaybackSpeed = (OneWayPlaybackDuration >= 0.0f) ? (1.0f / FMath::Max(OneWayPlaybackDuration, SMALL_NUMBER)) : 0.0f;
 		const float Direction = OutputPathedMovemenState.bIsPathProgressionIncreasing ? 1.0f : -1.0f;
 		const int32 NumFramesSinceLastChange = (Params.TimeStep.ServerFrame - OutputPathedMovemenState.LastChangeFrame) + 1;
 		NewPathProgress = OutputPathedMovemenState.LastChangePathProgress + NumFramesSinceLastChange * (Direction * PlaybackSpeed * FrameDuration);
