[Chaos Mover] Some fixes for based movement:
- Recompute the ground info at the start of each frame. This avoids having to use the previous frame version on resims.
- Move the relative velocity setting to just the walking mode.
- Store the ground velocity data when updating the ground to avoid having to recompute every time it's needed.
- Add specific code to deal with standing on a rotating platform. This is needed because the character and ground velocities are not the same in this case.

#Synced-CL 48652273

diff --git a/Engine/Plugins/Experimental/ChaosMover/Source/ChaosMover/Private/Backends/ChaosMoverBackend.cpp b/Engine/Plugins/Experimental/ChaosMover/Source/ChaosMover/Private/Backends/ChaosMoverBackend.cpp
--- a/Engine/Plugins/Experimental/ChaosMover/Source/ChaosMover/Private/Backends/ChaosMoverBackend.cpp
+++ b/Engine/Plugins/Experimental/ChaosMover/Source/ChaosMover/Private/Backends/ChaosMoverBackend.cpp
@@ -640,16 +640,14 @@ void UChaosMoverBackendComponent::FinalizeFrame(double ResultsTimeInMs)
 	FMoverDefaultSyncState& TransformSyncState = InterpolatedOutput.SyncState.SyncStateCollection.FindOrAddMutableDataByType<FMoverDefaultSyncState>();
 	if (Chaos::FPBDRigidParticle* Particle = GetControlledParticle())
 	{
-		TransformSyncState.SetTransforms_WorldSpace(Particle->GetX(),
+		TransformSyncState.SetTransforms_WorldSpace(
+			Particle->GetX(),
 			FRotator(Particle->GetR()),
 			Particle->GetV(),
 			FMath::RadiansToDegrees(Particle->GetW()),
 			TransformSyncState.GetMovementBase(),
 			TransformSyncState.GetMovementBaseBoneName());
 
-		// Make sure the move direction intent is in base space (the base quat is identity if there's no base, effectively making this a no-op)
-		TransformSyncState.MoveDirectionIntent = TransformSyncState.GetCapturedMovementBaseQuat().UnrotateVector(TransformSyncState.MoveDirectionIntent);
-	
 		check(Chaos::FVec3::IsNearlyEqual(Chaos::FVec3(TransformSyncState.GetLocation_WorldSpace()), Particle->GetX(), UE_KINDA_SMALL_NUMBER));
 	}
 
diff --git a/Engine/Plugins/Experimental/ChaosMover/Source/ChaosMover/Private/ChaosMoverConsoleVariables.cpp b/Engine/Plugins/Experimental/ChaosMover/Source/ChaosMover/Private/ChaosMoverConsoleVariables.cpp
--- a/Engine/Plugins/Experimental/ChaosMover/Source/ChaosMover/Private/ChaosMoverConsoleVariables.cpp
+++ b/Engine/Plugins/Experimental/ChaosMover/Source/ChaosMover/Private/ChaosMoverConsoleVariables.cpp
@@ -47,6 +47,10 @@ namespace UE::ChaosMover
 		FAutoConsoleVariableRef CVarChaosMoverFrameDifferenceLeniencyForEventTimeComparison(TEXT("ChaosMover.Networking.FrameDifferenceLeniencyForEventTimeComparison"),
 			FrameDifferenceLeniencyForEventTimeComparison, TEXT("Within this number of frames events are considered to be at the same time and will not trigger multiple times.\n"));
 
+		float FallingCheckRelativeSpeedLimit = 60.0f;
+		FAutoConsoleVariableRef CVarChaosMoverFallingCheckRelativeSpeedLimit(TEXT("ChaosMover.FallingCheckRelativeSpeedLimit"),
+			FallingCheckRelativeSpeedLimit, TEXT("Above this limit the falling check will consider the character to be lifting off the surface.\n"));
+
 
 		static FAutoConsoleCommandWithWorld TeleportOnlyClientCommand
 		(
diff --git a/Engine/Plugins/Experimental/ChaosMover/Source/ChaosMover/Private/ChaosMoverSimulation.cpp b/Engine/Plugins/Experimental/ChaosMover/Source/ChaosMover/Private/ChaosMoverSimulation.cpp
--- a/Engine/Plugins/Experimental/ChaosMover/Source/ChaosMover/Private/ChaosMoverSimulation.cpp
+++ b/Engine/Plugins/Experimental/ChaosMover/Source/ChaosMover/Private/ChaosMoverSimulation.cpp
@@ -378,36 +378,8 @@ void UChaosMoverSimulation::PreSimulationTickCharacter(const IChaosCharacterMove
 		}
 	}
 
-	// The state machine works in local space relative to the ground, so if we're on
-	// the ground make the state relative to the ground
-	// TODO - Can we use the local space in the sync state rather than subtracting and
-	//        re-adding ground velocity to the world space velocity?
-	FVector LocalGroundVelocity = FVector::ZeroVector;
-	const UBaseMovementMode* Mode = StateMachine.FindMovementMode(CurrentSyncState.MovementMode).Get();
-	if (Mode && Mode->HasGameplayTag(Mover_IsOnGround, true))
-	{
-		FMoverDefaultSyncState& PreSimDefaultSyncState = CurrentSyncState.SyncStateCollection.FindOrAddMutableDataByType<FMoverDefaultSyncState>();
-
-		FFloorCheckResult LastFloorResult;
-		if (Blackboard->TryGet(CommonBlackboard::LastFloorResult, LastFloorResult))
-		{
-			LocalGroundVelocity = FVector::VectorPlaneProject(
-				UChaosGroundMovementUtils::ComputeLocalGroundVelocity_Internal(PreSimDefaultSyncState.GetLocation_WorldSpace(), LastFloorResult),
-				LastFloorResult.HitResult.ImpactNormal);
-			
-			if (!LocalGroundVelocity.IsNearlyZero())
-			{
-				PreSimDefaultSyncState.SetTransforms_WorldSpace(
-					PreSimDefaultSyncState.GetLocation_WorldSpace(),
-					PreSimDefaultSyncState.GetOrientation_WorldSpace(),
-					PreSimDefaultSyncState.GetVelocity_WorldSpace() - LocalGroundVelocity,
-					PreSimDefaultSyncState.GetAngularVelocityDegrees_WorldSpace());
-			}
-		}
-	}
-
-	FChaosMoverGroundSimState& GroundSimState = InternalSimData.FindOrAddMutableDataByType<FChaosMoverGroundSimState>();
-	GroundSimState.LocalVelocity = LocalGroundVelocity;
+	// Update the floor before we run any update.
+	CharacterMode.UpdateCurrentFloor(TimeStep);
 }
 
 void UChaosMoverSimulation::OnSimulationRollback(const FMoverTimeStep& NewTimeStep, const FMoverSyncState& PrevSyncState)
@@ -482,17 +454,9 @@ void UChaosMoverSimulation::PostSimulationTickCharacter(const IChaosCharacterMov
 
 	if (Chaos::FPBDRigidParticleHandle* ParticleHandle = GetControlledParticle())
 	{
-		// Linear motion
+		// Constraint target
 		CharacterSimState.TargetDeltaPosition = PostSimDefaultSyncState->GetLocation_WorldSpace() - ParticleHandle->GetX();
 
-		FVector LocalGroundVelocity = FVector::ZeroVector;
-		if (const FChaosMoverGroundSimState* GroundSimState = InternalSimData.FindDataByType<FChaosMoverGroundSimState>())
-		{
-			LocalGroundVelocity = GroundSimState->LocalVelocity;
-		}
-		ParticleHandle->SetV(PostSimDefaultSyncState->GetVelocity_WorldSpace() + LocalGroundVelocity);
-
-		// Angular motion
 		FQuat TgtQuat = PostSimDefaultSyncState->GetOrientation_WorldSpace().Quaternion();
 		TgtQuat.EnforceShortestArcWith(ParticleHandle->GetR());
 		FQuat QuatRotation = TgtQuat * ParticleHandle->GetR().Inverse();
@@ -505,14 +469,9 @@ void UChaosMoverSimulation::PostSimulationTickCharacter(const IChaosCharacterMov
 		}
 		CharacterSimState.TargetDeltaFacing = AngularDisplacement.Dot(UpDir);
 
-		if (CharacterMode.ShouldCharacterRemainUpright())
-		{
-			const float DeltaTimeSeconds = TimeStep.StepMs * 0.001f;
-			if (DeltaTimeSeconds > UE_SMALL_NUMBER)
-			{
-				ParticleHandle->SetW(AngularDisplacement / DeltaTimeSeconds);
-			}
-		}
+		// Velocity
+		ParticleHandle->SetV(PostSimDefaultSyncState->GetVelocity_WorldSpace());
+		ParticleHandle->SetW(FMath::DegreesToRadians(PostSimDefaultSyncState->GetAngularVelocityDegrees_WorldSpace()));
 	}
 	else
 	{
@@ -520,14 +479,24 @@ void UChaosMoverSimulation::PostSimulationTickCharacter(const IChaosCharacterMov
 		CharacterSimState.TargetDeltaFacing = 0.0f;
 	}
 
-	// Update the movement base
+	// Update the floor result
 	check(Blackboard.Get());
 	FFloorCheckResult FloorResult;
 	bool FoundLastFloorResult = Blackboard->TryGet(CommonBlackboard::LastFloorResult, FloorResult);
-	bool FoundFloor = FoundLastFloorResult && FloorResult.bBlockingHit;
-	// Note: We want to record the movement base but we don't record the transform
-	// so don't use this to get a relative transform for the sync state
-	PostSimDefaultSyncState->SetMovementBase(FoundFloor ? FloorResult.HitResult.GetComponent() : nullptr);
+	if (FoundLastFloorResult && FloorResult.bBlockingHit)
+	{
+		const FVector WorldSpaceDirectionIntent = PostSimDefaultSyncState->GetIntent_WorldSpace();
+		PostSimDefaultSyncState->SetTransforms_WorldSpace(
+			PostSimDefaultSyncState->GetLocation_WorldSpace(),
+			PostSimDefaultSyncState->GetOrientation_WorldSpace(),
+			PostSimDefaultSyncState->GetVelocity_WorldSpace(),
+			PostSimDefaultSyncState->GetAngularVelocityDegrees_WorldSpace(),
+			FloorResult.HitResult.GetComponent(),
+			FloorResult.HitResult.BoneName);
+
+		// Make sure the move direction intent is in base space (the base quat is identity if there's no base, effectively making this a no-op)
+		PostSimDefaultSyncState->MoveDirectionIntent = PostSimDefaultSyncState->GetCapturedMovementBaseQuat().UnrotateVector(PostSimDefaultSyncState->MoveDirectionIntent);
+	}
 
 	FFloorResultData& FloorData = OutputData.AdditionalOutputData.FindOrAddMutableDataByType<FFloorResultData>();
 	FloorData.FloorResult = FloorResult;
diff --git a/Engine/Plugins/Experimental/ChaosMover/Source/ChaosMover/Private/ChaosMoverSimulationTypes.cpp b/Engine/Plugins/Experimental/ChaosMover/Source/ChaosMover/Private/ChaosMoverSimulationTypes.cpp
--- a/Engine/Plugins/Experimental/ChaosMover/Source/ChaosMover/Private/ChaosMoverSimulationTypes.cpp
+++ b/Engine/Plugins/Experimental/ChaosMover/Source/ChaosMover/Private/ChaosMoverSimulationTypes.cpp
@@ -1,13 +1,55 @@
 // Copyright Epic Games, Inc. All Rights Reserved.
 
 #include "ChaosMover/ChaosMoverSimulationTypes.h"
+
+#include "Chaos/PhysicsObjectInternalInterface.h"
 #include "ChaosMover/ChaosMoverSimulation.h"
 #include "ChaosMover/ChaosMoverLog.h"
 #include "ChaosMover/ChaosMoverConsoleVariables.h"
+#include "MoveLibrary/BasedMovementUtils.h"
 #include "Serialization/ArchiveCrc32.h"
 
 #include UE_INLINE_GENERATED_CPP_BY_NAME(ChaosMoverSimulationTypes)
 
+namespace UE::ChaosMover
+{
+	namespace Blackboard
+	{
+		const FName GroundDynamicsInfo = TEXT("GroundDynamicsInfo");
+	}
+
+	FGroundDynamicsInfo::FGroundDynamicsInfo()
+		: LinearVelocity(FVector::ZeroVector)
+		, AngularVelocityDegrees(FVector::ZeroVector)
+		, bIsMoving(0)
+		, bIsDynamic(0)
+		, bIsGravityEnabled(0)
+	{
+	}
+
+	FGroundDynamicsInfo::FGroundDynamicsInfo(const FFloorCheckResult& InFloorResult)
+		: FGroundDynamicsInfo()
+	{
+		if (Chaos::FPhysicsObjectHandle PhysicsObject = InFloorResult.HitResult.PhysicsObject)
+		{
+			Chaos::FReadPhysicsObjectInterface_Internal Interface = Chaos::FPhysicsObjectInternalInterface::GetRead();
+			if (const Chaos::FPBDRigidParticleHandle* ParticleHandle = Interface.GetRigidParticle(PhysicsObject))
+			{
+				bIsDynamic = ParticleHandle->IsDynamic();
+
+				FVector Offset = InFloorResult.HitResult.ImpactPoint - ParticleHandle->GetTransformXRCom().GetLocation();
+				Offset = FVector::VectorPlaneProject(Offset, InFloorResult.HitResult.ImpactNormal);
+
+				LinearVelocity = ParticleHandle->GetV() + ParticleHandle->GetW().Cross(Offset);
+				AngularVelocityDegrees = FMath::RadiansToDegrees(ParticleHandle->GetW());
+
+				bIsMoving = (LinearVelocity.SizeSquared() > UE_SMALL_NUMBER) || (AngularVelocityDegrees.SizeSquared() > UE_SMALL_NUMBER);
+				bIsGravityEnabled = ParticleHandle->GravityEnabled();
+			}
+		}
+	}
+}
+
 void FChaosMoverSimulationDefaultInputs::Reset()
 {
 	CollisionResponseParams = FCollisionResponseParams();
@@ -342,4 +384,4 @@ UScriptStruct* FDebugTeleportToInstantMovementEffect::GetScriptStruct() const
 FString FDebugTeleportToInstantMovementEffect::ToSimpleString() const
 {
 	return FString::Printf(TEXT("FDebugTeleportToInstantMovementEffect. TeleportLocation: %s\nTeleportRotation: %s"), *TeleportLocation.ToString(), *TeleportRotation.ToString());
-}
+}
\ No newline at end of file
diff --git a/Engine/Plugins/Experimental/ChaosMover/Source/ChaosMover/Private/Character/Effects/ChaosCharacterApplyVelocityEffect.cpp b/Engine/Plugins/Experimental/ChaosMover/Source/ChaosMover/Private/Character/Effects/ChaosCharacterApplyVelocityEffect.cpp
--- a/Engine/Plugins/Experimental/ChaosMover/Source/ChaosMover/Private/Character/Effects/ChaosCharacterApplyVelocityEffect.cpp
+++ b/Engine/Plugins/Experimental/ChaosMover/Source/ChaosMover/Private/Character/Effects/ChaosCharacterApplyVelocityEffect.cpp
@@ -61,7 +61,7 @@ bool FChaosCharacterApplyVelocityEffect::ApplyMovementEffect_Async(FApplyMovemen
 		break;
 	}
 
-	OutputSyncState.SetTransforms_WorldSpace(Position,Orientation,Velocity,FVector::ZeroVector);
+	OutputSyncState.SetTransforms_WorldSpace(Position, Orientation, Velocity, CurrentSyncState->GetAngularVelocityDegrees_WorldSpace());
 
 	return true;
 }
diff --git a/Engine/Plugins/Experimental/ChaosMover/Source/ChaosMover/Private/Character/Modes/ChaosCharacterMovementMode.cpp b/Engine/Plugins/Experimental/ChaosMover/Source/ChaosMover/Private/Character/Modes/ChaosCharacterMovementMode.cpp
--- a/Engine/Plugins/Experimental/ChaosMover/Source/ChaosMover/Private/Character/Modes/ChaosCharacterMovementMode.cpp
+++ b/Engine/Plugins/Experimental/ChaosMover/Source/ChaosMover/Private/Character/Modes/ChaosCharacterMovementMode.cpp
@@ -8,10 +8,14 @@
 #include "ChaosMover/Character/Settings/SharedChaosCharacterMovementSettings.h"
 #include "ChaosMover/ChaosMoverLog.h"
 #include "ChaosMover/ChaosMoverSimulation.h"
+#include "ChaosMover/Utilities/ChaosMoverQueryUtils.h"
 #include "Components/CapsuleComponent.h"
 #include "Components/SkeletalMeshComponent.h"
 #include "GameFramework/Character.h"
 #include "Math/UnitConversion.h"
+#include "MoveLibrary/FloorQueryUtils.h"
+#include "MoveLibrary/MovementUtils.h"
+#include "MoveLibrary/WaterMovementUtils.h"
 
 #include UE_INLINE_GENERATED_CPP_BY_NAME(ChaosCharacterMovementMode)
 
@@ -216,3 +220,58 @@ void UChaosCharacterMovementMode::ModifyContacts(const FMoverTimeStep& TimeStep,
 		}
 	}
 }
+
+void UChaosCharacterMovementMode::UpdateCurrentFloor(const FMoverTimeStep& TimeStep) const
+{
+	check(Simulation);
+
+	UMoverBlackboard* SimBlackboard = Simulation->GetBlackboard_Mutable();
+	if (!SimBlackboard)
+	{
+		return;
+	}
+
+	bool bFloorCheckSucceeded = false;
+
+	if (const FChaosMoverSimulationDefaultInputs* SimInputs = Simulation->GetLocalSimInput().FindDataByType<FChaosMoverSimulationDefaultInputs>())
+	{
+		Chaos::FReadPhysicsObjectInterface_Internal ReadInterface = Chaos::FPhysicsObjectInternalInterface::GetRead();
+		if (const Chaos::FPBDRigidParticleHandle* UpdatedParticle = ReadInterface.GetRigidParticle(SimInputs->PhysicsObject))
+		{
+			FFloorCheckResult FloorResult;
+			FWaterCheckResult WaterResult;
+
+			const float DeltaSeconds = TimeStep.StepMs * 0.001f;
+			const FVector VelocityWithGravity = UpdatedParticle->GetV() + UMovementUtils::ComputeVelocityFromGravity(SimInputs->Gravity, DeltaSeconds);
+
+			UE::ChaosMover::Utils::FFloorSweepParams SweepParams{
+				.ResponseParams = SimInputs->CollisionResponseParams,
+				.QueryParams = SimInputs->CollisionQueryParams,
+				.Location = UpdatedParticle->GetX(),
+				.DeltaPos = VelocityWithGravity * DeltaSeconds,
+				.UpDir = SimInputs->UpDir,
+				.World = SimInputs->World,
+				.QueryDistance = 1.5f * GetTargetHeight(),
+				.QueryRadius = FMath::Min(GetGroundQueryRadius(), FMath::Max(SimInputs->PawnCollisionRadius - 5.0f, 0.0f)),
+				.MaxWalkSlopeCosine = GetMaxWalkSlopeCosine(),
+				.TargetHeight = GetTargetHeight(),
+				.CollisionChannel = SimInputs->CollisionChannel
+			};
+
+			UE::ChaosMover::Utils::FloorSweep_Internal(SweepParams, FloorResult, WaterResult);
+
+			SimBlackboard->Set(CommonBlackboard::LastFloorResult, FloorResult);
+			SimBlackboard->Set(CommonBlackboard::LastWaterResult, WaterResult);
+			SimBlackboard->Set(UE::ChaosMover::Blackboard::GroundDynamicsInfo, UE::ChaosMover::FGroundDynamicsInfo(FloorResult));
+
+			bFloorCheckSucceeded = true;
+		}
+	}
+
+	if (!bFloorCheckSucceeded)
+	{
+		SimBlackboard->Invalidate(CommonBlackboard::LastFloorResult);
+		SimBlackboard->Invalidate(CommonBlackboard::LastWaterResult);
+		SimBlackboard->Invalidate(UE::ChaosMover::Blackboard::GroundDynamicsInfo);
+	}
+}
diff --git a/Engine/Plugins/Experimental/ChaosMover/Source/ChaosMover/Private/Character/Modes/ChaosFallingMode.cpp b/Engine/Plugins/Experimental/ChaosMover/Source/ChaosMover/Private/Character/Modes/ChaosFallingMode.cpp
--- a/Engine/Plugins/Experimental/ChaosMover/Source/ChaosMover/Private/Character/Modes/ChaosFallingMode.cpp
+++ b/Engine/Plugins/Experimental/ChaosMover/Source/ChaosMover/Private/Character/Modes/ChaosFallingMode.cpp
@@ -21,6 +21,7 @@ UChaosFallingMode::UChaosFallingMode(const FObjectInitializer& ObjectInitializer
 	, AirControlPercentage(0.4f)
 	, FallingDeceleration(200.0f)
 	, FallingLateralFriction(0.0f)
+	, FallingAngularFriction(0.1f)
 	, OverTerminalSpeedFallingDeceleration(800.0f)
 	, TerminalMovementPlaneSpeed(1500.0f)
 	, bShouldClampTerminalVerticalSpeed(true)
@@ -90,14 +91,33 @@ void UChaosFallingMode::GenerateMove_Implementation(const FMoverTickStartData& S
 	Params.MoveInput = FVector::VectorPlaneProject(Params.MoveInput, DefaultSimInputs->UpDir);
 
 	FRotator IntendedOrientation_WorldSpace;
-	// If there's no intent from input to change orientation, use the current orientation
+
+	// If there's no intent from input to change orientation, use the current orientation integrated with the current angular velocity
 	if (!CharacterInputs || CharacterInputs->OrientationIntent.IsNearlyZero())
 	{
 		IntendedOrientation_WorldSpace = StartingSyncState->GetOrientation_WorldSpace();
+
+		if (ShouldCharacterRemainUpright())
+		{
+			IntendedOrientation_WorldSpace = FRotator(FQuat::FindBetweenVectors(IntendedOrientation_WorldSpace.RotateVector(DefaultSimInputs->UpDir), DefaultSimInputs->UpDir) * IntendedOrientation_WorldSpace.Quaternion());
+		}
+
+		const float AngularFrictionFactor = FMath::Clamp(1.0f - FallingAngularFriction * DeltaSeconds, 0.0f, 1.0f);
+		FVector AngularVelocityToApply = AngularFrictionFactor * StartingSyncState->GetAngularVelocityDegrees_WorldSpace();
+		if (ShouldCharacterRemainUpright())
+		{
+			AngularVelocityToApply = AngularVelocityToApply.ProjectOnToNormal(DefaultSimInputs->UpDir);
+		}
+		IntendedOrientation_WorldSpace = FRotator(UMovementUtils::ApplyAngularVelocityToQuat(IntendedOrientation_WorldSpace.Quaternion(), AngularVelocityToApply, DeltaSeconds));
 	}
 	else
 	{
 		IntendedOrientation_WorldSpace = CharacterInputs->GetOrientationIntentDir_WorldSpace().ToOrientationRotator();
+
+		if (ShouldCharacterRemainUpright())
+		{
+			IntendedOrientation_WorldSpace = FRotator(FQuat::FindBetweenVectors(IntendedOrientation_WorldSpace.RotateVector(DefaultSimInputs->UpDir), DefaultSimInputs->UpDir) * IntendedOrientation_WorldSpace.Quaternion());
+		}
 	}
 
 	FQuat WorldToGravityTransform = FQuat::FindBetweenNormals(FVector::UpVector, DefaultSimInputs->UpDir);
@@ -147,32 +167,6 @@ void UChaosFallingMode::GenerateMove_Implementation(const FMoverTickStartData& S
 	{
 		UMovementUtils::SetGravityVerticalComponent(OutProposedMove.LinearVelocity, VelocityWithGravity.Dot(DefaultSimInputs->UpDir), DefaultSimInputs->UpDir);
 	}
-
-	FFloorCheckResult FloorResult;
-	FWaterCheckResult WaterResult;
-	UMoverBlackboard* SimBlackboard = Simulation->GetBlackboard_Mutable();
-	if (SimBlackboard)
-	{
-		// Update the floor
-		UE::ChaosMover::Utils::FFloorSweepParams SweepParams{
-			.ResponseParams = DefaultSimInputs->CollisionResponseParams,
-			.QueryParams = DefaultSimInputs->CollisionQueryParams,
-			.Location = StartingSyncState->GetLocation_WorldSpace(),
-			.DeltaPos = OutProposedMove.LinearVelocity * TimeStep.StepMs * 0.001f,
-			.UpDir = DefaultSimInputs->UpDir,
-			.World = DefaultSimInputs->World,
-			.QueryDistance = 1.5f * GetTargetHeight(),
-			.QueryRadius = FMath::Min(GetGroundQueryRadius(), FMath::Max(DefaultSimInputs->PawnCollisionRadius - 5.0f, 0.0f)),
-			.MaxWalkSlopeCosine = GetMaxWalkSlopeCosine(),
-			.TargetHeight = GetTargetHeight(),
-			.CollisionChannel = DefaultSimInputs->CollisionChannel
-		};
-
-		UE::ChaosMover::Utils::FloorSweep_Internal(SweepParams, FloorResult, WaterResult);
-
-		SimBlackboard->Set(CommonBlackboard::LastFloorResult, FloorResult);
-		SimBlackboard->Set(CommonBlackboard::LastWaterResult, WaterResult);
-	}
 }
 
 void UChaosFallingMode::SimulationTick_Implementation(const FSimulationTickParams& Params, FMoverTickEndData& OutputState)
@@ -208,6 +202,7 @@ void UChaosFallingMode::SimulationTick_Implementation(const FSimulationTickParam
 
 	OutputState.MovementEndState.RemainingMs = 0.0f;
 	OutputState.MovementEndState.NextModeName = Params.StartState.SyncState.MovementMode;
+	
 	OutputSyncState.MoveDirectionIntent = ProposedMove.bHasDirIntent ? ProposedMove.DirectionIntent : FVector::ZeroVector;
 	OutputSyncState.SetTransforms_WorldSpace(
 		TargetPos,
diff --git a/Engine/Plugins/Experimental/ChaosMover/Source/ChaosMover/Private/Character/Modes/ChaosFlyingMode.cpp b/Engine/Plugins/Experimental/ChaosMover/Source/ChaosMover/Private/Character/Modes/ChaosFlyingMode.cpp
--- a/Engine/Plugins/Experimental/ChaosMover/Source/ChaosMover/Private/Character/Modes/ChaosFlyingMode.cpp
+++ b/Engine/Plugins/Experimental/ChaosMover/Source/ChaosMover/Private/Character/Modes/ChaosFlyingMode.cpp
@@ -138,4 +138,20 @@ void UChaosFlyingMode::SimulationTick_Implementation(const FSimulationTickParams
 		TargetOrient,
 		TargetVel,
 		ProposedMove.AngularVelocityDegrees);
+}
+
+void UChaosFlyingMode::UpdateCurrentFloor(const FMoverTimeStep&) const
+{
+	if (!Simulation)
+	{
+		UE_LOG(LogChaosMover, Warning, TEXT("No Simulation set on ChaosFlyingMode"));
+		return;
+	}
+
+	if (UMoverBlackboard* SimBlackboard = Simulation->GetBlackboard_Mutable())
+	{
+		SimBlackboard->Invalidate(CommonBlackboard::LastFloorResult);
+		SimBlackboard->Invalidate(CommonBlackboard::LastWaterResult);
+		SimBlackboard->Invalidate(UE::ChaosMover::Blackboard::GroundDynamicsInfo);
+	}
 }
\ No newline at end of file
diff --git a/Engine/Plugins/Experimental/ChaosMover/Source/ChaosMover/Private/Character/Modes/ChaosWalkingMode.cpp b/Engine/Plugins/Experimental/ChaosMover/Source/ChaosMover/Private/Character/Modes/ChaosWalkingMode.cpp
--- a/Engine/Plugins/Experimental/ChaosMover/Source/ChaosMover/Private/Character/Modes/ChaosWalkingMode.cpp
+++ b/Engine/Plugins/Experimental/ChaosMover/Source/ChaosMover/Private/Character/Modes/ChaosWalkingMode.cpp
@@ -78,10 +78,21 @@ void UChaosWalkingMode::GenerateMove_Implementation(const FMoverTickStartData& S
 	// Try to use the floor as the basis for the intended move direction (i.e. try to walk along slopes, rather than into them)
 	UMoverBlackboard* SimBlackboard = Simulation->GetBlackboard_Mutable();
 	FFloorCheckResult LastFloorResult;
+	UE::ChaosMover::FGroundDynamicsInfo GroundDynamicsInfo;
 	FVector MovementNormal = UpDirection;
-	if (!bMaintainHorizontalGroundVelocity && SimBlackboard && SimBlackboard->TryGet(CommonBlackboard::LastFloorResult, LastFloorResult) && LastFloorResult.IsWalkableFloor())
+	bool bWalkableFloor = false;
+
+	if (SimBlackboard)
 	{
-		MovementNormal = LastFloorResult.HitResult.ImpactNormal;
+		SimBlackboard->TryGet(CommonBlackboard::LastFloorResult, LastFloorResult);
+		SimBlackboard->TryGet(UE::ChaosMover::Blackboard::GroundDynamicsInfo, GroundDynamicsInfo);
+
+		bWalkableFloor = LastFloorResult.IsWalkableFloor();
+
+		if (!bMaintainHorizontalGroundVelocity && SimBlackboard->TryGet(CommonBlackboard::LastFloorResult, LastFloorResult) && bWalkableFloor)
+		{
+			MovementNormal = LastFloorResult.HitResult.ImpactNormal;
+		}
 	}
 
 	FRotator IntendedOrientation_WorldSpace;
@@ -97,7 +108,7 @@ void UChaosWalkingMode::GenerateMove_Implementation(const FMoverTickStartData& S
 
 	const FQuat WorldToGravityTransform = FQuat::FindBetweenNormals(FVector::UpVector, DefaultSimInputs->UpDir);
 	IntendedOrientation_WorldSpace = UMovementUtils::ApplyGravityToOrientationIntent(IntendedOrientation_WorldSpace, WorldToGravityTransform, bShouldCharacterRemainUpright);
-	const FVector CurrentVelocity = StartingSyncState->GetVelocity_WorldSpace();
+	const FVector CurrentVelocity = StartingSyncState->GetVelocity_WorldSpace() - GroundDynamicsInfo.LinearVelocity;
 
 	const FChaosMoverRequestedMoveInputs* RequestedMove = StartState.InputCmd.InputCollection.FindDataByType<FChaosMoverRequestedMoveInputs>();
 	if (RequestedMove)
@@ -127,10 +138,7 @@ void UChaosWalkingMode::GenerateMove_Implementation(const FMoverTickStartData& S
 		if (CharacterInputs)
 		{
 			Params.MoveInputType = CharacterInputs->GetMoveInputType();
-
-			const bool bMaintainInputMagnitude = true;
 			Params.MoveInput = CharacterInputs->GetMoveInput_WorldSpace();
-			//Params.MoveInput = UPlanarConstraintUtils::ConstrainDirectionToPlane(MoverComp->GetPlanarConstraint(), CharacterInputs->GetMoveInput_WorldSpace(), bMaintainInputMagnitude);
 		}
 		else
 		{
@@ -165,29 +173,27 @@ void UChaosWalkingMode::GenerateMove_Implementation(const FMoverTickStartData& S
 		OutProposedMove = UGroundMovementUtils::ComputeControlledGroundMove(Params);
 	}
 
-	if (SimBlackboard)
+	if (bMaintainHorizontalGroundVelocity && bWalkableFloor)
 	{
-		// Update the floor result and check the proposed move to prevent movement onto unwalkable surfaces
-		FVector OutDeltaPos = FVector::ZeroVector;
-		FFloorCheckResult FloorResult;
-		FWaterCheckResult WaterResult;
-		GetFloorAndCheckMovement(*StartingSyncState, OutProposedMove, *DefaultSimInputs, DeltaSeconds, FloorResult, WaterResult, OutDeltaPos);
-
-		OutProposedMove.LinearVelocity = (DeltaSeconds > 0.0f) ? (OutDeltaPos / DeltaSeconds) : FVector::ZeroVector;
-
-		SimBlackboard->Set(CommonBlackboard::LastFloorResult, FloorResult);
-		SimBlackboard->Set(CommonBlackboard::LastWaterResult, WaterResult);
+		// So far have assumed we are on level ground, so now add velocity up the slope
+		const double CosAngleBetweenUpAndNormal = UpDirection.Dot(LastFloorResult.HitResult.ImpactNormal);
+		if (CosAngleBetweenUpAndNormal > UE_SMALL_NUMBER)
+		{
+			OutProposedMove.LinearVelocity -= UpDirection * OutProposedMove.LinearVelocity.Dot(LastFloorResult.HitResult.ImpactNormal) / CosAngleBetweenUpAndNormal;
+		}
+	}
 
-		if (bMaintainHorizontalGroundVelocity)
+	// Special case for standing on a rotating platform. Since the relative ground velocity is not quite right
+	// in this case just set the propose move to zero if there is no input and the proposed move is small enough
+	if ((GroundDynamicsInfo.AngularVelocityDegrees.Dot(DefaultSimInputs->UpDir) > 1.0f) && CharacterInputs->GetMoveInput_WorldSpace().IsNearlyZero())
+	{
+		const float ErrorLimit = GetMaxSpeed() * 0.1f;
+		if (OutProposedMove.LinearVelocity.SizeSquared() < ErrorLimit * ErrorLimit)
 		{
-			// So far have assumed we are on level ground, so now add velocity up the slope
-			const double CosAngleBetweenUpAndNormal = UpDirection.Dot(FloorResult.HitResult.ImpactNormal);
-			if (CosAngleBetweenUpAndNormal > UE_SMALL_NUMBER)
-			{
-				OutProposedMove.LinearVelocity -= UpDirection * OutProposedMove.LinearVelocity.Dot(FloorResult.HitResult.ImpactNormal) / CosAngleBetweenUpAndNormal;
-			}
+			OutProposedMove.LinearVelocity = FVector::ZeroVector;
 		}
-	};
+
+	}
 }
 
 void UChaosWalkingMode::SimulationTick_Implementation(const FSimulationTickParams& Params, FMoverTickEndData& OutputState)
@@ -215,20 +221,31 @@ void UChaosWalkingMode::SimulationTick_Implementation(const FSimulationTickParam
 	FMoverDefaultSyncState& OutputSyncState = OutputState.SyncState.SyncStateCollection.FindOrAddMutableDataByType<FMoverDefaultSyncState>();
 	OutputSyncState = *StartingSyncState;
 
+	// Update the ground based on the proposed move
 	const float DeltaSeconds = Params.TimeStep.StepMs * 0.001f;
-	const FVector UpDirection = DefaultSimInputs->UpDir;
-
-	FVector GroundNormal = UpDirection;
-	const UMoverBlackboard* SimBlackboard = Simulation->GetBlackboard();
+	FVector DeltaPos;
 	FFloorCheckResult FloorResult;
-	if (SimBlackboard)
+	FWaterCheckResult WaterResult;
+	GetFloorAndCheckMovement(OutputSyncState, ProposedMove, *DefaultSimInputs, DeltaSeconds, FloorResult, WaterResult, DeltaPos);
+
+	ProposedMove.LinearVelocity = (DeltaSeconds > 0.0f) ? (DeltaPos / DeltaSeconds) : FVector::ZeroVector;
+
+	UE::ChaosMover::FGroundDynamicsInfo GroundDynamicsInfo(FloorResult);
+
+	if (UMoverBlackboard* SimBlackboard = Simulation->GetBlackboard_Mutable())
 	{
-		SimBlackboard->TryGet(CommonBlackboard::LastFloorResult, FloorResult);
-		GroundNormal = FloorResult.HitResult.ImpactNormal;
-	};
+		SimBlackboard->Set(CommonBlackboard::LastFloorResult, FloorResult);
+		SimBlackboard->Set(CommonBlackboard::LastWaterResult, WaterResult);
+		SimBlackboard->Set(UE::ChaosMover::Blackboard::GroundDynamicsInfo, GroundDynamicsInfo);
+	}
+
+	const FVector UpDirection = DefaultSimInputs->UpDir;
+	FVector GroundNormal = UpDirection;
 
 	if (FloorResult.IsWalkableFloor())
 	{
+		GroundNormal = FloorResult.HitResult.ImpactNormal;
+
 		FVector TargetPosition = StartingSyncState->GetLocation_WorldSpace();
 
 		// The base movement mode does not apply gravity in walking mode so apply here.
@@ -261,22 +278,21 @@ void UChaosWalkingMode::SimulationTick_Implementation(const FSimulationTickParam
 			TargetPosition += ProposedMovePlaneVelocity * DeltaSeconds;
 		}
 
-		FVector ProjectedGroundVelocity = UChaosGroundMovementUtils::ComputeLocalGroundVelocity_Internal(StartingSyncState->GetLocation_WorldSpace(), FloorResult);
-		const Chaos::FPBDRigidParticleHandle* GroundParticle = UChaosGroundMovementUtils::GetRigidParticleHandleFromFloorResult_Internal(FloorResult);
-		if (GroundParticle && GroundParticle->IsDynamic() && GroundParticle->GravityEnabled())
+		FVector ProjectedGroundVelocity = GroundDynamicsInfo.LinearVelocity;
+		if (GroundDynamicsInfo.bIsGravityEnabled)
 		{
 			// This might not be correct if different physics objects have different gravity but is saves having to go
 			// to the component to get the gravity on the physics volume.
 			ProjectedGroundVelocity += DefaultSimInputs->PhysicsObjectGravity * DefaultSimInputs->UpDir * DeltaSeconds;
 		}
-		const bool bIsGroundMoving = ProjectedGroundVelocity.SizeSquared() > UE_KINDA_SMALL_NUMBER;
+
 		const FVector ProjectedRelativeVelocity = TargetVelocity - ProjectedGroundVelocity;
 		const float ProjectedRelativeNormalVelocity = FloorResult.HitResult.ImpactNormal.Dot(TargetVelocity - ProjectedGroundVelocity);
 		const float ProjectedRelativeVerticalVelocity = GroundNormal.Dot(TargetVelocity - ProjectedGroundVelocity);
 		const float VerticalVelocityLimit = bNormalVelocityIntent ? 2.0f / DeltaSeconds : FMath::Abs(GroundNormal.Dot(DefaultSimInputs->Gravity) * DeltaSeconds);
 
 		bool bIsLiftingOffSurface = false;
-		if ((ProjectedRelativeNormalVelocity > VerticalVelocityLimit) && bIsGroundMoving && (ProjectedRelativeVerticalVelocity > VerticalVelocityLimit))
+		if ((ProjectedRelativeNormalVelocity > VerticalVelocityLimit) && GroundDynamicsInfo.bIsMoving && (ProjectedRelativeVerticalVelocity > VerticalVelocityLimit))
 		{
 			bIsLiftingOffSurface = true;
 		}
@@ -297,7 +313,7 @@ void UChaosWalkingMode::SimulationTick_Implementation(const FSimulationTickParam
 		}
 
 		// Put the target position on the floor at the target height
-		if (!bIsGroundMoving)
+		if (!GroundDynamicsInfo.bIsMoving)
 		{
 			TargetPosition -= UpDirection * InitialHeightAboveFloor;
 		}
@@ -306,15 +322,16 @@ void UChaosWalkingMode::SimulationTick_Implementation(const FSimulationTickParam
 		// This is always applied regardless of whether the character is supported
 		const FRotator TargetOrientation = UMovementUtils::ApplyAngularVelocityToRotator(StartingSyncState->GetOrientation_WorldSpace(), ProposedMove.AngularVelocityDegrees, DeltaSeconds);
 
-		OutputSyncState.SetTransforms_WorldSpace(TargetPosition,
+		OutputSyncState.MoveDirectionIntent = ProposedMove.bHasDirIntent ? ProposedMove.DirectionIntent : FVector::ZeroVector;
+		OutputSyncState.SetTransforms_WorldSpace(
+			TargetPosition,
 			TargetOrientation,
 			TargetVelocity,
-			ProposedMove.AngularVelocityDegrees);
+			StartingSyncState->GetAngularVelocityDegrees_WorldSpace());
 	}
 	
 	OutputState.MovementEndState.RemainingMs = 0.0f;
 	OutputState.MovementEndState.NextModeName = Params.StartState.SyncState.MovementMode;
-	OutputSyncState.MoveDirectionIntent = ProposedMove.bHasDirIntent ? ProposedMove.DirectionIntent : FVector::ZeroVector;
 }
 
 bool UChaosWalkingMode::CanStepUpOnHitSurface(const FFloorCheckResult& FloorResult, const FVector& CharacterVelocity) const
@@ -509,7 +526,7 @@ void UChaosWalkingMode::ModifyContacts(const FMoverTimeStep& TimeStep, const UE:
 	const float StepDistance = FMath::Abs(GetTargetHeight() - FloorResult.FloorDist);
 	if (StepDistance >= UE_KINDA_SMALL_NUMBER)
 	{
-		MinContactHeightStepUps = CharacterHeight - GetTargetHeight() + SharedSettings->MaxStepHeight;
+		MinContactHeightStepUps = CharacterHeight - GetTargetHeight() + (1.0f + GroundDamping) * SharedSettings->MaxStepHeight;
 	}
 
 	for (Chaos::FContactPairModifier& PairModifier : Modifier.GetContacts(UpdatedParticle))
diff --git a/Engine/Plugins/Experimental/ChaosMover/Source/ChaosMover/Private/Character/Transitions/ChaosCharacterFallingCheck.cpp b/Engine/Plugins/Experimental/ChaosMover/Source/ChaosMover/Private/Character/Transitions/ChaosCharacterFallingCheck.cpp
--- a/Engine/Plugins/Experimental/ChaosMover/Source/ChaosMover/Private/Character/Transitions/ChaosCharacterFallingCheck.cpp
+++ b/Engine/Plugins/Experimental/ChaosMover/Source/ChaosMover/Private/Character/Transitions/ChaosCharacterFallingCheck.cpp
@@ -3,6 +3,7 @@
 #include "ChaosMover/Character/Transitions/ChaosCharacterFallingCheck.h"
 
 #include "Chaos/ParticleHandle.h"
+#include "ChaosMover/ChaosMoverConsoleVariables.h"
 #include "ChaosMover/ChaosMoverLog.h"
 #include "ChaosMover/ChaosMoverSimulation.h"
 #include "ChaosMover/ChaosMoverSimulationTypes.h"
@@ -57,7 +58,8 @@ FTransitionEvalResult UChaosCharacterFallingCheck::Evaluate_Implementation(const
 		check(Blackboard);
 
 		FFloorCheckResult FloorResult;
-		if (!Blackboard->TryGet(CommonBlackboard::LastFloorResult, FloorResult))
+		UE::ChaosMover::FGroundDynamicsInfo GroundDynamicsInfo;
+		if (!(Blackboard->TryGet(CommonBlackboard::LastFloorResult, FloorResult) && Blackboard->TryGet(UE::ChaosMover::Blackboard::GroundDynamicsInfo, GroundDynamicsInfo)))
 		{
 			return EvalResult;
 		}
@@ -73,10 +75,8 @@ FTransitionEvalResult UChaosCharacterFallingCheck::Evaluate_Implementation(const
 			if (const IChaosCharacterMovementModeInterface* Mode = Cast<IChaosCharacterMovementModeInterface>(BaseMode.Get()))
 			{
 				const float DeltaSeconds = Params.TimeStep.StepMs * 0.001f;
-				const FVector LinearVelocity = StartingSyncState->GetVelocity_WorldSpace();
-
 				const FVector ProjectedVelocity = StartingSyncState->GetVelocity_WorldSpace() + SimInputs->Gravity * DeltaSeconds;
-				FVector ProjectedGroundVelocity = UChaosGroundMovementUtils::ComputeLocalGroundVelocity_Internal(StartingSyncState->GetLocation_WorldSpace(), FloorResult);
+				FVector ProjectedGroundVelocity = GroundDynamicsInfo.LinearVelocity;
 				const Chaos::FPBDRigidParticleHandle* GroundParticle = UChaosGroundMovementUtils::GetRigidParticleHandleFromFloorResult_Internal(FloorResult);
 				if (GroundParticle && GroundParticle->IsDynamic() && GroundParticle->GravityEnabled())
 				{
@@ -84,11 +84,11 @@ FTransitionEvalResult UChaosCharacterFallingCheck::Evaluate_Implementation(const
 					// to the component to get the gravity on the physics volume.
 					ProjectedGroundVelocity += SimInputs->PhysicsObjectGravity * SimInputs->UpDir * DeltaSeconds;
 				}
-				const bool bIsGroundMoving = ProjectedGroundVelocity.SizeSquared() > UE_KINDA_SMALL_NUMBER;
+				const bool bIsGroundMoving = FMath::Abs(ProjectedGroundVelocity.Dot(SimInputs->UpDir)) > UE_KINDA_SMALL_NUMBER;
 				const FVector ProjectedRelativeVelocity = ProjectedVelocity - ProjectedGroundVelocity;
 				const float ProjectedRelativeNormalVelocity = FloorResult.HitResult.ImpactNormal.Dot(ProjectedVelocity - ProjectedGroundVelocity);
 				const float ProjectedRelativeVerticalVelocity = SimInputs->UpDir.Dot(ProjectedVelocity - ProjectedGroundVelocity);
-				const float VerticalVelocityLimit = 2.0f / DeltaSeconds;
+				const float VerticalVelocityLimit = UE::ChaosMover::CVars::FallingCheckRelativeSpeedLimit;
 				bool bIsLiftingOffSurface = false;
 				if ((ProjectedRelativeNormalVelocity > VerticalVelocityLimit) && bIsGroundMoving && (ProjectedRelativeVerticalVelocity > VerticalVelocityLimit))
 				{
@@ -130,4 +130,4 @@ FTransitionEvalResult UChaosCharacterFallingCheck::Evaluate_Implementation(const
 
 void UChaosCharacterFallingCheck::Trigger_Implementation(const FSimulationTickParams& Params)
 {
-}
\ No newline at end of file
+}
diff --git a/Engine/Plugins/Experimental/ChaosMover/Source/ChaosMover/Private/Character/Transitions/ChaosCharacterJumpCheck.cpp b/Engine/Plugins/Experimental/ChaosMover/Source/ChaosMover/Private/Character/Transitions/ChaosCharacterJumpCheck.cpp
--- a/Engine/Plugins/Experimental/ChaosMover/Source/ChaosMover/Private/Character/Transitions/ChaosCharacterJumpCheck.cpp
+++ b/Engine/Plugins/Experimental/ChaosMover/Source/ChaosMover/Private/Character/Transitions/ChaosCharacterJumpCheck.cpp
@@ -69,15 +69,15 @@ void UChaosCharacterJumpCheck::Trigger_Implementation(const FSimulationTickParam
 
 		// Add on the local ground velocity to make the jump velocity relative to the ground
 		FFloorCheckResult FloorResult;
+		UE::ChaosMover::FGroundDynamicsInfo GroundDynamicsInfo;
 		const UMoverBlackboard* Blackboard = Simulation->GetBlackboard();
-		if (Blackboard && Blackboard->TryGet(CommonBlackboard::LastFloorResult, FloorResult))
+		if (Blackboard && Blackboard->TryGet(CommonBlackboard::LastFloorResult, FloorResult) && Blackboard->TryGet(UE::ChaosMover::Blackboard::GroundDynamicsInfo, GroundDynamicsInfo))
 		{
 			if (Chaos::FPhysicsObjectHandle GroundObject = FloorResult.HitResult.PhysicsObject)
 			{
 				if (StartingSyncState)
 				{
-					const FVector LocalGroundVelocity = UChaosGroundMovementUtils::ComputeLocalGroundVelocity_Internal(StartingSyncState->GetLocation_WorldSpace(), FloorResult);
-					JumpVelocity += LocalGroundVelocity.Dot(SimInputs->UpDir) * SimInputs->UpDir;
+					JumpVelocity += GroundDynamicsInfo.LinearVelocity.Dot(SimInputs->UpDir) * SimInputs->UpDir;
 				}
 			}
 		}
diff --git a/Engine/Plugins/Experimental/ChaosMover/Source/ChaosMover/Private/Character/Transitions/ChaosCharacterLandingCheck.cpp b/Engine/Plugins/Experimental/ChaosMover/Source/ChaosMover/Private/Character/Transitions/ChaosCharacterLandingCheck.cpp
--- a/Engine/Plugins/Experimental/ChaosMover/Source/ChaosMover/Private/Character/Transitions/ChaosCharacterLandingCheck.cpp
+++ b/Engine/Plugins/Experimental/ChaosMover/Source/ChaosMover/Private/Character/Transitions/ChaosCharacterLandingCheck.cpp
@@ -55,11 +55,12 @@ FTransitionEvalResult UChaosCharacterLandingCheck::Evaluate_Implementation(const
 				if (!TransitionToGroundMode.IsNone())
 				{
 					FFloorCheckResult FloorResult;
-					if (Blackboard->TryGet(CommonBlackboard::LastFloorResult, FloorResult) && FloorResult.IsWalkableFloor())
+					UE::ChaosMover::FGroundDynamicsInfo GroundDynamicsInfo;
+					if (Blackboard->TryGet(CommonBlackboard::LastFloorResult, FloorResult) && FloorResult.IsWalkableFloor() && Blackboard->TryGet(UE::ChaosMover::Blackboard::GroundDynamicsInfo, GroundDynamicsInfo))
 					{
 						const float TargetHeight = Mode->GetTargetHeight();
 						const float FloorDistanceWithFloorNormal = FloorResult.HitResult.ImpactNormal.Dot(StartingSyncState->GetLocation_WorldSpace() - FloorResult.HitResult.ImpactPoint);
-						const FVector LocalGroundVelocity = UChaosGroundMovementUtils::ComputeLocalGroundVelocity_Internal(StartingSyncState->GetLocation_WorldSpace(), FloorResult);
+						const FVector LocalGroundVelocity = GroundDynamicsInfo.LinearVelocity;
 						const float RelativeVerticalVelocity = FloorResult.HitResult.ImpactNormal.Dot(LinearVelocity - LocalGroundVelocity);
 						const float ProjectedFloorDistance = FloorDistanceWithFloorNormal + RelativeVerticalVelocity * DeltaSeconds;
 						const bool bIsFloorWithinReach = ProjectedFloorDistance < TargetHeight + FloorDistanceTolerance + UE_KINDA_SMALL_NUMBER;
diff --git a/Engine/Plugins/Experimental/ChaosMover/Source/ChaosMover/Public/ChaosMover/ChaosMoverConsoleVariables.h b/Engine/Plugins/Experimental/ChaosMover/Source/ChaosMover/Public/ChaosMover/ChaosMoverConsoleVariables.h
--- a/Engine/Plugins/Experimental/ChaosMover/Source/ChaosMover/Public/ChaosMover/ChaosMoverConsoleVariables.h
+++ b/Engine/Plugins/Experimental/ChaosMover/Source/ChaosMover/Public/ChaosMover/ChaosMoverConsoleVariables.h
@@ -16,5 +16,6 @@ namespace UE::ChaosMover
 		CHAOSMOVER_API extern bool bDisableResimDuplicateEventChecking;
 		CHAOSMOVER_API extern int32 InstantMovementEffectIDHistorySize;
 		CHAOSMOVER_API extern int32 FrameDifferenceLeniencyForEventTimeComparison;
+		CHAOSMOVER_API extern float FallingCheckRelativeSpeedLimit;
 	}
 };
diff --git a/Engine/Plugins/Experimental/ChaosMover/Source/ChaosMover/Public/ChaosMover/ChaosMoverSimulationTypes.h b/Engine/Plugins/Experimental/ChaosMover/Source/ChaosMover/Public/ChaosMover/ChaosMoverSimulationTypes.h
--- a/Engine/Plugins/Experimental/ChaosMover/Source/ChaosMover/Public/ChaosMover/ChaosMoverSimulationTypes.h
+++ b/Engine/Plugins/Experimental/ChaosMover/Source/ChaosMover/Public/ChaosMover/ChaosMoverSimulationTypes.h
@@ -7,6 +7,8 @@
 #include "Chaos/PhysicsObject.h"
 #include "CollisionQueryParams.h"
 #include "DefaultMovementSet/MovementModifiers/StanceModifier.h"
+#include "MoveLibrary/FloorQueryUtils.h"
+#include "MoveLibrary/WaterMovementUtils.h"
 #include "MoverSimulationTypes.h"
 #include "StructUtils/InstancedStruct.h"
 
@@ -52,6 +54,22 @@ namespace UE::ChaosMover
 	// without including ChaosMoverSimulation.h
 	CHAOSMOVER_API FMoverDataCollection& GetDebugSimData(UChaosMoverSimulation* Simulation);
 
+	namespace Blackboard
+	{
+		extern CHAOSMOVER_API const FName GroundDynamicsInfo;
+	}
+
+	struct FGroundDynamicsInfo
+	{
+		CHAOSMOVER_API FGroundDynamicsInfo();
+		CHAOSMOVER_API FGroundDynamicsInfo(const FFloorCheckResult& InFloorResult);
+		FVector LinearVelocity;
+		FVector AngularVelocityDegrees;
+		uint32 bIsMoving : 1;
+		uint32 bIsDynamic : 1;
+		uint32 bIsGravityEnabled : 1;
+	};
+
 } // namespace UE::ChaosMover
 
 // Default chaos mover character simulation state, contains state basic for character simulation
@@ -241,6 +259,9 @@ public:
 
 	UFUNCTION(BlueprintCallable, Category = ChaosMover)
 	virtual void ClearAccelerationOverride() = 0;
+
+	UFUNCTION(BlueprintCallable, Category = ChaosMover)
+	virtual void UpdateCurrentFloor(const FMoverTimeStep& TimeStep) const = 0;
 };
 
 // Interface for mover modes moving on ground like characters, using a character ground constraint
diff --git a/Engine/Plugins/Experimental/ChaosMover/Source/ChaosMover/Public/ChaosMover/Character/Modes/ChaosCharacterMovementMode.h b/Engine/Plugins/Experimental/ChaosMover/Source/ChaosMover/Public/ChaosMover/Character/Modes/ChaosCharacterMovementMode.h
--- a/Engine/Plugins/Experimental/ChaosMover/Source/ChaosMover/Public/ChaosMover/Character/Modes/ChaosCharacterMovementMode.h
+++ b/Engine/Plugins/Experimental/ChaosMover/Source/ChaosMover/Public/ChaosMover/Character/Modes/ChaosCharacterMovementMode.h
@@ -72,6 +72,9 @@ public:
 	CHAOSMOVER_API virtual void OverrideAcceleration(float Value) override;
 	CHAOSMOVER_API virtual void ClearAccelerationOverride() override;
 
+	// Updates the floor result in the blackboard
+	CHAOSMOVER_API virtual void UpdateCurrentFloor(const FMoverTimeStep& TimeStep) const override;
+
 protected:
 	CHAOSMOVER_API void SetTargetHeightOverride(float InTargetHeight);
 	CHAOSMOVER_API void ClearTargetHeightOverride();
diff --git a/Engine/Plugins/Experimental/ChaosMover/Source/ChaosMover/Public/ChaosMover/Character/Modes/ChaosFallingMode.h b/Engine/Plugins/Experimental/ChaosMover/Source/ChaosMover/Public/ChaosMover/Character/Modes/ChaosFallingMode.h
--- a/Engine/Plugins/Experimental/ChaosMover/Source/ChaosMover/Public/ChaosMover/Character/Modes/ChaosFallingMode.h
+++ b/Engine/Plugins/Experimental/ChaosMover/Source/ChaosMover/Public/ChaosMover/Character/Modes/ChaosFallingMode.h
@@ -52,6 +52,14 @@ public:
 	UPROPERTY(Category = Mover, EditAnywhere, BlueprintReadWrite, meta = (ClampMin = "0"))
 	float FallingLateralFriction;
 
+	/**
+	 * Friction to apply to rotational air movement.
+	 * Angular velocity is scaled each tick by a factor (1-Friction * TimeStep) so friction
+	 * values greater than 1/TimeStep will result in all angular velocity being removed.
+	 */
+	UPROPERTY(Category = Mover, EditAnywhere, BlueprintReadWrite, meta = (ClampMin = "0"))
+	float FallingAngularFriction;
+
 	/**
 	 * Deceleration to apply to air movement when falling faster than terminal velocity
 	 * Note: This is NOT applied to vertical velocity, only movement plane velocity
diff --git a/Engine/Plugins/Experimental/ChaosMover/Source/ChaosMover/Public/ChaosMover/Character/Modes/ChaosFlyingMode.h b/Engine/Plugins/Experimental/ChaosMover/Source/ChaosMover/Public/ChaosMover/Character/Modes/ChaosFlyingMode.h
--- a/Engine/Plugins/Experimental/ChaosMover/Source/ChaosMover/Public/ChaosMover/Character/Modes/ChaosFlyingMode.h
+++ b/Engine/Plugins/Experimental/ChaosMover/Source/ChaosMover/Public/ChaosMover/Character/Modes/ChaosFlyingMode.h
@@ -6,8 +6,6 @@
 
 #include "ChaosFlyingMode.generated.h"
 
-#define UE_API CHAOSMOVER_API
-
 /**
  * Chaos character flying mode
  */
@@ -17,10 +15,10 @@ class UChaosFlyingMode : public UChaosCharacterMovementMode
 	GENERATED_BODY()
 
 public:
-	UE_API UChaosFlyingMode(const FObjectInitializer& ObjectInitializer);
+	CHAOSMOVER_API UChaosFlyingMode(const FObjectInitializer& ObjectInitializer);
 
-	UE_API virtual void GenerateMove_Implementation(const FMoverTickStartData& StartState, const FMoverTimeStep& TimeStep, FProposedMove& OutProposedMove) const override;
-	UE_API virtual void SimulationTick_Implementation(const FSimulationTickParams& Params, FMoverTickEndData& OutputState) override;
-};
+	CHAOSMOVER_API virtual void GenerateMove_Implementation(const FMoverTickStartData& StartState, const FMoverTimeStep& TimeStep, FProposedMove& OutProposedMove) const override;
+	CHAOSMOVER_API virtual void SimulationTick_Implementation(const FSimulationTickParams& Params, FMoverTickEndData& OutputState) override;
 
-#undef UE_API
+	CHAOSMOVER_API virtual void UpdateCurrentFloor(const FMoverTimeStep& TimeStep) const override;
+};
