Mover: Making requested input struct properties VisibleAnywhere so they shoe up on the Chaos Visual Debugger

#Synced-CL 48614097

diff --git a/Engine/Plugins/Experimental/ChaosMover/Source/ChaosMover/Public/ChaosMover/Character/ChaosCharacterInputs.h b/Engine/Plugins/Experimental/ChaosMover/Source/ChaosMover/Public/ChaosMover/Character/ChaosCharacterInputs.h
--- a/Engine/Plugins/Experimental/ChaosMover/Source/ChaosMover/Public/ChaosMover/Character/ChaosCharacterInputs.h
+++ b/Engine/Plugins/Experimental/ChaosMover/Source/ChaosMover/Public/ChaosMover/Character/ChaosCharacterInputs.h
@@ -170,13 +170,13 @@ struct FChaosMoverRequestedMoveInputs : public FMoverDataStructBase
 
 public:
 	// Requested velocity in cm/s
-	UPROPERTY(BlueprintReadWrite, Category = "Chaos Mover")
+	UPROPERTY(VisibleAnywhere, BlueprintReadWrite, Category = "Chaos Mover")
 	FVector RequestedVelocity = FVector::ZeroVector;
 
-	UPROPERTY(BlueprintReadWrite, Category = "Chaos Mover")
+	UPROPERTY(VisibleAnywhere, BlueprintReadWrite, Category = "Chaos Mover")
 	bool bForceMaxSpeed = false;
 
-	UPROPERTY(BlueprintReadWrite, Category = "Chaos Mover")
+	UPROPERTY(VisibleAnywhere, BlueprintReadWrite, Category = "Chaos Mover")
 	bool bUseAcceleration = true;
 
 	virtual UScriptStruct* GetScriptStruct() const override
