ChaosMover: Adding an option to falling mode where we can always clamp the requested velocity to max speed. This is off by default, we're itnroducing it to fix an issue with Pimlico where AI would feed arbitrary large requested velocities.

#Synced-CL 48614097

#Synced-CL 47488690

diff --git a/Engine/Plugins/Experimental/ChaosMover/Source/ChaosMover/Private/Character/Modes/ChaosFallingMode.cpp b/Engine/Plugins/Experimental/ChaosMover/Source/ChaosMover/Private/Character/Modes/ChaosFallingMode.cpp
--- a/Engine/Plugins/Experimental/ChaosMover/Source/ChaosMover/Private/Character/Modes/ChaosFallingMode.cpp
+++ b/Engine/Plugins/Experimental/ChaosMover/Source/ChaosMover/Private/Character/Modes/ChaosFallingMode.cpp
@@ -77,7 +77,10 @@ void UChaosFallingMode::GenerateMove_Implementation(const FMoverTickStartData& S
 	if (RequestedMove)
 	{
 		Params.MoveInputType = EMoveInputType::Velocity;
-		Params.MoveInput = RequestedMove->bForceMaxSpeed ? RequestedMove->RequestedVelocity : RequestedMove->RequestedVelocity.GetClampedToMaxSize(GetMaxSpeed());
+		// It's a bit confusing but bForceMaxSpeed is ill named and means that the requested velocity will NOT be clamped to max speed
+		// bAlwaysClampRequestedVelocityToMaxSpeed takes precedence over bForceMaxSpeed
+		const bool bClampToMaxSpeed = bAlwaysClampRequestedVelocityToMaxSpeed || !RequestedMove->bForceMaxSpeed;
+		Params.MoveInput = bClampToMaxSpeed ? RequestedMove->RequestedVelocity.GetClampedToMaxSize(GetMaxSpeed()) : RequestedMove->RequestedVelocity;
 	}
 	else
 	{
diff --git a/Engine/Plugins/Experimental/ChaosMover/Source/ChaosMover/Public/ChaosMover/Character/Modes/ChaosFallingMode.h b/Engine/Plugins/Experimental/ChaosMover/Source/ChaosMover/Public/ChaosMover/Character/Modes/ChaosFallingMode.h
--- a/Engine/Plugins/Experimental/ChaosMover/Source/ChaosMover/Public/ChaosMover/Character/Modes/ChaosFallingMode.h
+++ b/Engine/Plugins/Experimental/ChaosMover/Source/ChaosMover/Public/ChaosMover/Character/Modes/ChaosFallingMode.h
@@ -89,6 +89,12 @@ public:
 	UPROPERTY(Category = Mover, EditAnywhere, BlueprintReadWrite, meta = (ClampMin = "0", ForceUnits = "cm/s"))
 	float TerminalVerticalSpeed;
 
+	/**
+	 * If true and input requests a velocity, whether to always clamp that velocity to max speed, regardless of input's bForceMaxSpeed flag
+	 */
+	UPROPERTY(Category = Mover, EditAnywhere, BlueprintReadWrite)
+	bool bAlwaysClampRequestedVelocityToMaxSpeed = false;
+
 protected:
 
 	/*
