Mover: Adding context to BaseMovementMode Activate/Deactivate functions so we can tell whether it is being called in the course of regular forward simulation, or during forward resimulation, or due to a state rollback (rewind).
(There is one last piece of information we could add later on which is whether it is due to a regular local forward simulation or due to a target being applied, which could modulate a bit how we react.)

#Synced-CL 48734469

diff --git a/Engine/Plugins/Experimental/ChaosMover/Source/ChaosMover/Private/ChaosMoverSimulation.cpp b/Engine/Plugins/Experimental/ChaosMover/Source/ChaosMover/Private/ChaosMoverSimulation.cpp
--- a/Engine/Plugins/Experimental/ChaosMover/Source/ChaosMover/Private/ChaosMoverSimulation.cpp
+++ b/Engine/Plugins/Experimental/ChaosMover/Source/ChaosMover/Private/ChaosMoverSimulation.cpp
@@ -751,7 +751,7 @@ void UChaosMoverSimulation::Teleport(const FTransform& TargetTransform, const FM
 	if (IsCharacterConstraintEnabled())
 	{
 		DisableCharacterConstraint();
-		StateMachine.SetModeImmediately(StateMachine.GetDefaultModeName());		
+		StateMachine.SetModeImmediately(TimeStep, StateMachine.GetDefaultModeName());
 	}
 
 	// Invalidate floor results. Should there be a way to indicate which blackboard results to invalidate for which type of operation, like rewind, teleport, etc?
diff --git a/Engine/Plugins/Experimental/ChaosMover/Source/ChaosMover/Private/ChaosMoverStateMachine.cpp b/Engine/Plugins/Experimental/ChaosMover/Source/ChaosMover/Private/ChaosMoverStateMachine.cpp
--- a/Engine/Plugins/Experimental/ChaosMover/Source/ChaosMover/Private/ChaosMoverStateMachine.cpp
+++ b/Engine/Plugins/Experimental/ChaosMover/Source/ChaosMover/Private/ChaosMoverStateMachine.cpp
@@ -73,8 +73,6 @@ namespace UE::ChaosMover
 
 	void FMoverStateMachine::ClearAllMovementModes()
 	{
-		Modes.Empty();
-
 		for (TPair<FName,TWeakObjectPtr<UBaseMovementMode>>& Element : Modes)
 		{
 			TStrongObjectPtr<UBaseMovementMode> Mode = Element.Value.Pin();
@@ -84,6 +82,7 @@ namespace UE::ChaosMover
 				Mode->OnUnregistered();
 			}
 		}
+		Modes.Empty();
 
 		ConstructDefaultModes();	// Note that we're resetting to our defaults so we keep the null movement mode
 	}
@@ -161,10 +160,10 @@ namespace UE::ChaosMover
 		}
 	}
 
-	void FMoverStateMachine::SetModeImmediately(FName DesiredModeName, bool bShouldReenter)
+	void FMoverStateMachine::SetModeImmediately(const FMoverTimeStep& TimeStep, FName DesiredModeName, bool bShouldReenter)
 	{
 		QueueNextModeInternal(DesiredModeName, bShouldReenter, /*bIsFromRollback*/ false);
-		AdvanceToNextMode();
+		AdvanceToNextMode(TimeStep);
 	}
 
 	void FMoverStateMachine::ClearQueuedMode()
@@ -350,7 +349,7 @@ namespace UE::ChaosMover
 			QueueNextMode(WorkingSubstepStartData.SyncState.MovementMode);
 		}
 
-		AdvanceToNextMode();
+		AdvanceToNextMode(TimeStep);
 
 		int SubStepCount = 0;
 		const int32 MaxConsecutiveFullRefundedSubsteps = GetDefault<UMoverDeveloperSettings>()->MaxTimesToRefundSubstep;
@@ -400,7 +399,7 @@ namespace UE::ChaosMover
 				if (CurrentModeName != OutputState.SyncState.MovementMode)
 				{
 					bModeSetFromInstantEffect = true;
-					SetModeImmediately(OutputState.SyncState.MovementMode);
+					SetModeImmediately(SubTimeStep, OutputState.SyncState.MovementMode);
 					WorkingSubstepStartData.SyncState.MovementMode = CurrentModeName;
 				}
 			}
@@ -449,7 +448,7 @@ namespace UE::ChaosMover
 
 			if (bHasLayeredMoveContributions && !CombinedLayeredMove.PreferredMode.IsNone() && !bModeSetFromInstantEffect)
 			{
-				SetModeImmediately(CombinedLayeredMove.PreferredMode);
+				SetModeImmediately(SubTimeStep, CombinedLayeredMove.PreferredMode);
 				OutputState.SyncState.MovementMode = CurrentModeName;
 			}
 
@@ -588,7 +587,7 @@ namespace UE::ChaosMover
 			}
 
 			// Switch modes if necessary (note that this will allow exit/enter on the same state)
-			AdvanceToNextMode();
+			AdvanceToNextMode(SubTimeStep);
 			OutputState.SyncState.MovementMode = CurrentModeName;
 
 			for (TSharedPtr<FMovementModifierBase> Modifier : ActiveModifiers)
@@ -622,7 +621,7 @@ namespace UE::ChaosMover
 			{
 				if (CurrentModeName != OutputState.SyncState.MovementMode)
 				{
-					SetModeImmediately(OutputState.SyncState.MovementMode);
+					SetModeImmediately(TimeStep, OutputState.SyncState.MovementMode);
 				}
 			}
 		}
@@ -786,13 +785,19 @@ namespace UE::ChaosMover
 		ClearQueuedMode();
 	}
 
-	void FMoverStateMachine::AdvanceToNextMode()
+	void FMoverStateMachine::AdvanceToNextMode(const FMoverTimeStep& TimeStep)
 	{
 		QUICK_SCOPE_CYCLE_COUNTER(MoverStateMachine_AdvanceToNextMode);
 
 		TStrongObjectPtr<UImmediateMovementModeTransition> QueuedModeTransition = QueuedModeTransitionWeakPtr.Pin();
 		if (QueuedModeTransition)
 		{
+			TStrongObjectPtr<UChaosMoverSimulation> SimStrongObjPtr = Simulation.Pin();
+			if (!ensure(SimStrongObjPtr))
+			{
+				return;
+			}
+
 			const FName NextModeName = QueuedModeTransition->GetNextModeName();
 
 			if (NextModeName != NAME_None)
@@ -809,13 +814,22 @@ namespace UE::ChaosMover
 						const FName PreviousModeName = CurrentModeName;
 						CurrentModeName = NextModeName;
 
+						FMoverEventContext MoverEventContext;
+						MoverEventContext.EventTimeMs = TimeStep.BaseSimTimeMs;
+						MoverEventContext.bIsDuringResimulation = TimeStep.bIsResimulating;
+						MoverEventContext.ServerFrame = TimeStep.ServerFrame;
+						MoverEventContext.bIsCausedByRollback = QueuedModeTransition->IsCausedByRollback();
+
 						if (PreviousModeName != NAME_None && Modes.Contains(PreviousModeName))
 						{
-							Modes[PreviousModeName]->Deactivate();
+							if (TStrongObjectPtr<UBaseMovementMode> PreviousMode = Modes[PreviousModeName].Pin())
+							{
+								PreviousMode->Deactivate(MoverEventContext);
+							}
 						}
 
 						// Track last mode change in the blackboard
-						if (URollbackBlackboard_InternalWrapper* RollbackBlackboard = Simulation->GetRollbackBlackboard_Mutable())
+						if (URollbackBlackboard_InternalWrapper* RollbackBlackboard = SimStrongObjPtr->GetRollbackBlackboard_Mutable())
 						{
 							FMovementModeChangeRecord ModeChangeRecord;
 							ModeChangeRecord.ModeName = CurrentModeName;
@@ -826,16 +840,16 @@ namespace UE::ChaosMover
 							RollbackBlackboard->TrySet(CommonBlackboard::LastModeChangeRecord, ModeChangeRecord);
 						}
 
-						Modes[CurrentModeName]->Activate();
+						if (TStrongObjectPtr<UBaseMovementMode> CurrentMode = Modes[CurrentModeName].Pin())
+						{	
+							CurrentMode->Activate(MoverEventContext);
+						}
 
 						// Notify the simulation of a mode change so it can react accordingly
-						if (TStrongObjectPtr<UChaosMoverSimulation> SimStrongObjPtr = Simulation.Pin())
-						{
-							TSharedPtr<FMovementModeChangedEventData> ModeChangedEventPtr = MakeShared<FMovementModeChangedEventData>(InternalSimTime, PreviousModeName, NextModeName);
-							ModeChangedEventPtr->Context.bIsCausedByRollback = QueuedModeTransition->IsCausedByRollback();
+						TSharedPtr<FMovementModeChangedEventData> ModeChangedEventPtr = MakeShared<FMovementModeChangedEventData>(InternalSimTime, PreviousModeName, NextModeName);
+						ModeChangedEventPtr->Context.bIsCausedByRollback = QueuedModeTransition->IsCausedByRollback();
 
-							SimStrongObjPtr->AddEvent(ModeChangedEventPtr);
-						}
+						SimStrongObjPtr->AddEvent(ModeChangedEventPtr);
 					}
 				}
 			}
@@ -873,7 +887,7 @@ namespace UE::ChaosMover
 			const FScheduledInstantMovementEffect& ScheduledEffect = QueuedInstantEffects[EffectIndex].ScheduledEffect;
 			if (ScheduledEffect.ShouldExecuteAtFrame(TimeStep.ServerFrame))
 			{
-				bInstantMovementEffectApplied |= ScheduledEffect.Effect->ApplyMovementEffect_Async(EffectParams_Async, OutputState);
+				bInstantMovementEffectApplied |= ScheduledEffect.Effect.IsValid() && ScheduledEffect.Effect->ApplyMovementEffect_Async(EffectParams_Async, OutputState);
 
 #if !defined(BUILD_SHIPPING) || !BUILD_SHIPPING
 				if (const FChaosMoverSimulationDefaultInputs* SimInputs = Simulation->GetLocalSimInput().FindDataByType<FChaosMoverSimulationDefaultInputs>())
@@ -886,17 +900,17 @@ namespace UE::ChaosMover
 				// We remove this effect since it has been applied
 				QueuedInstantEffects.RemoveAt(EffectIndex);
 			}
-#if !defined(BUILD_SHIPPING) || !BUILD_SHIPPING
 			else
 			{
 				EffectIndex++;
+#if !defined(BUILD_SHIPPING) || !BUILD_SHIPPING
 				if (const FChaosMoverSimulationDefaultInputs* SimInputs = Simulation->GetLocalSimInput().FindDataByType<FChaosMoverSimulationDefaultInputs>())
 				{
 					UE_LOG(LogChaosMover, Verbose, TEXT("(%s) SKIPPING Instant Effect scheduled for frame %d at frame %d: %s"),
 						*ToString(SimInputs->OwningActor->GetNetMode()), ScheduledEffect.ExecutionServerFrame, TimeStep.ServerFrame, ScheduledEffect.Effect.IsValid() ? *ScheduledEffect.Effect->ToSimpleString() : TEXT("INVALID INSTANT EFFECT"));
 				}
-			}
 #endif // !defined(BUILD_SHIPPING) || !BUILD_SHIPPING
+			}
 		}
 
 		return bInstantMovementEffectApplied;
diff --git a/Engine/Plugins/Experimental/ChaosMover/Source/ChaosMover/Public/ChaosMover/ChaosMoverStateMachine.h b/Engine/Plugins/Experimental/ChaosMover/Source/ChaosMover/Public/ChaosMover/ChaosMoverStateMachine.h
--- a/Engine/Plugins/Experimental/ChaosMover/Source/ChaosMover/Public/ChaosMover/ChaosMoverStateMachine.h
+++ b/Engine/Plugins/Experimental/ChaosMover/Source/ChaosMover/Public/ChaosMover/ChaosMoverStateMachine.h
@@ -40,7 +40,7 @@ namespace UE::ChaosMover
 		CHAOSMOVER_API void SetDefaultMode(FName NewDefaultModeName);
 
 		CHAOSMOVER_API void QueueNextMode(FName DesiredNextModeName, bool bShouldReenter = false);
-		CHAOSMOVER_API void SetModeImmediately(FName DesiredModeName, bool bShouldReenter = false);
+		CHAOSMOVER_API void SetModeImmediately(const FMoverTimeStep& TimeStep, FName DesiredModeName, bool bShouldReenter = false);
 		CHAOSMOVER_API void ClearQueuedMode();
 
 		CHAOSMOVER_API void OnSimulationTick(const FMoverTimeStep& TimeStep, const FMoverTickStartData& StartState, UMoverBlackboard* SimBlackboard, UMovementMixer* MovementMixer, FMoverTickEndData& OutputState);
@@ -116,7 +116,7 @@ namespace UE::ChaosMover
 
 	private:
 		void ConstructDefaultModes();
-		void AdvanceToNextMode();
+		void AdvanceToNextMode(const FMoverTimeStep& TimeStep);
 		void FlushQueuedMovesToGroup(FLayeredMoveGroup& Group);
 		void FlushQueuedModifiersToGroup(FMovementModifierGroup& ModifierGroup);
 		void FlushModifierCancellationsToGroup(FMovementModifierGroup& ActiveModifierGroup);
diff --git a/Engine/Plugins/Experimental/Mover/Source/Mover/Private/DefaultMovementSet/Modes/AsyncNavWalkingMode.cpp b/Engine/Plugins/Experimental/Mover/Source/Mover/Private/DefaultMovementSet/Modes/AsyncNavWalkingMode.cpp
--- a/Engine/Plugins/Experimental/Mover/Source/Mover/Private/DefaultMovementSet/Modes/AsyncNavWalkingMode.cpp
+++ b/Engine/Plugins/Experimental/Mover/Source/Mover/Private/DefaultMovementSet/Modes/AsyncNavWalkingMode.cpp
@@ -411,9 +411,9 @@ void UAsyncNavWalkingMode::SetTurnGeneratorClass(TSubclassOf<UObject> TurnGenera
 	}
 }
 
-void UAsyncNavWalkingMode::Activate()
+void UAsyncNavWalkingMode::Activate(const FMoverEventContext& Context)
 {
-	Super::Activate();
+	Super::Activate(Context);
 	
 	if (UMoverComponent* MoverComp = GetMoverComponent())
 	{
diff --git a/Engine/Plugins/Experimental/Mover/Source/Mover/Private/DefaultMovementSet/Modes/NavWalkingMode.cpp b/Engine/Plugins/Experimental/Mover/Source/Mover/Private/DefaultMovementSet/Modes/NavWalkingMode.cpp
--- a/Engine/Plugins/Experimental/Mover/Source/Mover/Private/DefaultMovementSet/Modes/NavWalkingMode.cpp
+++ b/Engine/Plugins/Experimental/Mover/Source/Mover/Private/DefaultMovementSet/Modes/NavWalkingMode.cpp
@@ -416,18 +416,18 @@ void UNavWalkingMode::SetCollisionForNavWalking(bool bEnable)
 	}
 }
 
-void UNavWalkingMode::Activate()
+void UNavWalkingMode::Activate(const FMoverEventContext& Context)
 {
-	Super::Activate();
+	Super::Activate(Context);
 	SetCollisionForNavWalking(true);
 
 	NavDataInterface = GetNavData();
 }
 
-void UNavWalkingMode::Deactivate()
+void UNavWalkingMode::Deactivate(const FMoverEventContext& Context)
 {
 	SetCollisionForNavWalking(false);
-	Super::Deactivate();
+	Super::Deactivate(Context);
 }
 
 const INavigationDataInterface* UNavWalkingMode::GetNavData() const
diff --git a/Engine/Plugins/Experimental/Mover/Source/Mover/Private/MovementMode.cpp b/Engine/Plugins/Experimental/Mover/Source/Mover/Private/MovementMode.cpp
--- a/Engine/Plugins/Experimental/Mover/Source/Mover/Private/MovementMode.cpp
+++ b/Engine/Plugins/Experimental/Mover/Source/Mover/Private/MovementMode.cpp
@@ -61,7 +61,7 @@ void UBaseMovementMode::OnUnregistered()
 	K2_OnUnregistered();
 }
 
-void UBaseMovementMode::Activate()
+void UBaseMovementMode::Activate(const FMoverEventContext& Context)
 {
 	if (!bSupportsAsync)
 	{
@@ -69,7 +69,7 @@ void UBaseMovementMode::Activate()
 	}
 }
 
-void UBaseMovementMode::Deactivate()
+void UBaseMovementMode::Deactivate(const FMoverEventContext& Context)
 {
 	if (!bSupportsAsync)
 	{
diff --git a/Engine/Plugins/Experimental/Mover/Source/Mover/Private/MovementModeStateMachine.cpp b/Engine/Plugins/Experimental/Mover/Source/Mover/Private/MovementModeStateMachine.cpp
--- a/Engine/Plugins/Experimental/Mover/Source/Mover/Private/MovementModeStateMachine.cpp
+++ b/Engine/Plugins/Experimental/Mover/Source/Mover/Private/MovementModeStateMachine.cpp
@@ -73,7 +73,13 @@ void UMovementModeStateMachine::ClearAllMovementModes()
 {
 	if (Modes.Contains(CurrentModeName))
 	{
-		Modes[CurrentModeName]->Deactivate();		
+		FMoverEventContext MoverEventContext;
+		MoverEventContext.EventTimeMs = CurrentBaseTimeStep.BaseSimTimeMs;
+		MoverEventContext.ServerFrame = CurrentBaseTimeStep.ServerFrame;
+		MoverEventContext.bIsDuringResimulation = false;
+		MoverEventContext.bIsCausedByRollback = false;
+
+		Modes[CurrentModeName]->Deactivate(MoverEventContext);
 	}
 
 	for (TPair<FName, TObjectPtr<UBaseMovementMode>>& Element : Modes)
@@ -727,6 +733,7 @@ void UMovementModeStateMachine::AdvanceToNextMode()
 {
 	const FName NextModeName = QueuedModeTransition->GetNextModeName();
 	const bool bShouldNextModeReenter = QueuedModeTransition->ShouldReenter();
+	const bool bIsCausedByRollback = QueuedModeTransition->IsCausedByRollback();
 
 	ClearQueuedMode();
 
@@ -740,9 +747,15 @@ void UMovementModeStateMachine::AdvanceToNextMode()
 		const FName PreviousModeName = CurrentModeName;
 		CurrentModeName = NextModeName;
 
+		FMoverEventContext MoverEventContext;
+		MoverEventContext.EventTimeMs = CurrentBaseTimeStep.BaseSimTimeMs;
+		MoverEventContext.ServerFrame = CurrentBaseTimeStep.ServerFrame;
+		MoverEventContext.bIsDuringResimulation = CurrentBaseTimeStep.bIsResimulating;
+		MoverEventContext.bIsCausedByRollback = bIsCausedByRollback;
+
 		if (PreviousModeName != NAME_None && Modes.Contains(PreviousModeName))
 		{
-			Modes[PreviousModeName]->Deactivate();
+			Modes[PreviousModeName]->Deactivate(MoverEventContext);
 		}
 
 		// signal movement mode change event
@@ -759,7 +772,7 @@ void UMovementModeStateMachine::AdvanceToNextMode()
 			RollbackBlackboard->TrySet(CommonBlackboard::LastModeChangeRecord, ModeChangeRecord);
 		}
 
-		Modes[CurrentModeName]->Activate();
+		Modes[CurrentModeName]->Activate(MoverEventContext);
 
 		MoverComp->OnMovementModeChanged.Broadcast(PreviousModeName, NextModeName);
 	}
diff --git a/Engine/Plugins/Experimental/Mover/Source/Mover/Public/DefaultMovementSet/Modes/AsyncNavWalkingMode.h b/Engine/Plugins/Experimental/Mover/Source/Mover/Public/DefaultMovementSet/Modes/AsyncNavWalkingMode.h
--- a/Engine/Plugins/Experimental/Mover/Source/Mover/Public/DefaultMovementSet/Modes/AsyncNavWalkingMode.h
+++ b/Engine/Plugins/Experimental/Mover/Source/Mover/Public/DefaultMovementSet/Modes/AsyncNavWalkingMode.h
@@ -29,7 +29,7 @@ class UAsyncNavWalkingMode : public UBaseMovementMode
 public:
 	UE_API UAsyncNavWalkingMode();
 
-	UE_API virtual void Activate() override;
+	UE_API virtual void Activate(const FMoverEventContext& Context) override;
 	UE_API virtual void SimulationTick_Implementation(const FSimulationTickParams& Params, FMoverTickEndData& OutputState) override;
 	UE_API virtual void GenerateMove_Implementation(const FMoverTickStartData& StartState, const FMoverTimeStep& TimeStep, FProposedMove& OutProposedMove) const override;
 
diff --git a/Engine/Plugins/Experimental/Mover/Source/Mover/Public/DefaultMovementSet/Modes/NavWalkingMode.h b/Engine/Plugins/Experimental/Mover/Source/Mover/Public/DefaultMovementSet/Modes/NavWalkingMode.h
--- a/Engine/Plugins/Experimental/Mover/Source/Mover/Public/DefaultMovementSet/Modes/NavWalkingMode.h
+++ b/Engine/Plugins/Experimental/Mover/Source/Mover/Public/DefaultMovementSet/Modes/NavWalkingMode.h
@@ -49,8 +49,8 @@ class UNavWalkingMode : public UBaseMovementMode
 public:
 	UE_API UNavWalkingMode();
 	
-	UE_API virtual void Activate() override;
-	UE_API virtual void Deactivate() override;
+	UE_API virtual void Activate(const FMoverEventContext& Context) override;
+	UE_API virtual void Deactivate(const FMoverEventContext& Context) override;
 	UE_API virtual void SimulationTick_Implementation(const FSimulationTickParams& Params, FMoverTickEndData& OutputState) override;
 	UE_API virtual void GenerateMove_Implementation(const FMoverTickStartData& StartState, const FMoverTimeStep& TimeStep, FProposedMove& OutProposedMove) const override;
 
diff --git a/Engine/Plugins/Experimental/Mover/Source/Mover/Public/MovementMode.h b/Engine/Plugins/Experimental/Mover/Source/Mover/Public/MovementMode.h
--- a/Engine/Plugins/Experimental/Mover/Source/Mover/Public/MovementMode.h
+++ b/Engine/Plugins/Experimental/Mover/Source/Mover/Public/MovementMode.h
@@ -47,8 +47,8 @@ public:
 	UE_API virtual void OnUnregistered();
 	
 	// These functions are called immediately when the state machine switches modes
-	UE_API virtual void Activate();
-	UE_API virtual void Deactivate();
+	UE_API virtual void Activate(const FMoverEventContext& Context);
+	UE_API virtual void Deactivate(const FMoverEventContext& Context);
 
 	// These functions are called when the sync state is changed on the game thread
 	// and a new mode is activated/deactivated
