[Chaos Mover] Update the constraint settings every frame instead of just when the mode changes

#Synced-CL 48756964

diff --git a/Engine/Plugins/Experimental/ChaosMover/Source/ChaosMover/Private/ChaosMoverSimulation.cpp b/Engine/Plugins/Experimental/ChaosMover/Source/ChaosMover/Private/ChaosMoverSimulation.cpp
--- a/Engine/Plugins/Experimental/ChaosMover/Source/ChaosMover/Private/ChaosMoverSimulation.cpp
+++ b/Engine/Plugins/Experimental/ChaosMover/Source/ChaosMover/Private/ChaosMoverSimulation.cpp
@@ -525,6 +525,7 @@ void UChaosMoverSimulation::PostSimulationTickCharacterConstraint(const IChaosCh
 	Chaos::FCharacterGroundConstraintSettings& Settings = CharacterConstraintHandle->GetSettings_Mutable();
 	if (const FChaosMoverSimulationDefaultInputs* SimInputs = LocalSimInput.FindDataByType<FChaosMoverSimulationDefaultInputs>())
 	{
+		CharacterConstraintMode.UpdateConstraintSettings(Settings);
 		Settings.VerticalAxis = SimInputs->UpDir;
 		Settings.TargetHeight = CharacterConstraintMode.GetTargetHeight();
 	}
@@ -958,9 +959,6 @@ void UChaosMoverSimulation::OnMovementModeChanged(const FMovementModeChangedEven
 		{
 			if (CharacterConstraintHandle)
 			{
-				Chaos::FCharacterGroundConstraintSettings& Settings = CharacterConstraintHandle->GetSettings_Mutable();
-				NextCharacterConstraintMode->UpdateConstraintSettings(Settings);
-
 				// Character ground constraint modes currently assume moving a dynamic particle and using a character ground constraint
 				// Revise if we start supporting moving a character kinematically
 				if (!IsControlledParticleDynamic())
