[Chaos Mover] Don't apply any angular velocity to target if the character should not remain upright and add option to never apply angular velocity to target in a falling mode

#Synced-CL 48756964

diff --git a/Engine/Plugins/Experimental/ChaosMover/Source/ChaosMover/Private/ChaosMoverSimulation.cpp b/Engine/Plugins/Experimental/ChaosMover/Source/ChaosMover/Private/ChaosMoverSimulation.cpp
--- a/Engine/Plugins/Experimental/ChaosMover/Source/ChaosMover/Private/ChaosMoverSimulation.cpp
+++ b/Engine/Plugins/Experimental/ChaosMover/Source/ChaosMover/Private/ChaosMoverSimulation.cpp
@@ -471,7 +471,10 @@ void UChaosMoverSimulation::PostSimulationTickCharacter(const IChaosCharacterMov
 
 		// Velocity
 		ParticleHandle->SetV(PostSimDefaultSyncState->GetVelocity_WorldSpace());
-		ParticleHandle->SetW(FMath::DegreesToRadians(PostSimDefaultSyncState->GetAngularVelocityDegrees_WorldSpace()));
+		if (CharacterMode.ShouldCharacterRemainUpright())
+		{
+			ParticleHandle->SetW(FMath::DegreesToRadians(PostSimDefaultSyncState->GetAngularVelocityDegrees_WorldSpace()));
+		}
 	}
 	else
 	{
diff --git a/Engine/Plugins/Experimental/ChaosMover/Source/ChaosMover/Private/Character/Modes/ChaosFallingMode.cpp b/Engine/Plugins/Experimental/ChaosMover/Source/ChaosMover/Private/Character/Modes/ChaosFallingMode.cpp
--- a/Engine/Plugins/Experimental/ChaosMover/Source/ChaosMover/Private/Character/Modes/ChaosFallingMode.cpp
+++ b/Engine/Plugins/Experimental/ChaosMover/Source/ChaosMover/Private/Character/Modes/ChaosFallingMode.cpp
@@ -21,7 +21,7 @@ UChaosFallingMode::UChaosFallingMode(const FObjectInitializer& ObjectInitializer
 	, AirControlPercentage(0.4f)
 	, FallingDeceleration(200.0f)
 	, FallingLateralFriction(0.0f)
-	, FallingAngularFriction(0.1f)
+	, FallingAngularFriction(30.0f)
 	, OverTerminalSpeedFallingDeceleration(800.0f)
 	, TerminalMovementPlaneSpeed(1500.0f)
 	, bShouldClampTerminalVerticalSpeed(true)
@@ -96,7 +96,7 @@ void UChaosFallingMode::GenerateMove_Implementation(const FMoverTickStartData& S
 	FRotator IntendedOrientation_WorldSpace;
 
 	// If there's no intent from input to change orientation, use the current orientation integrated with the current angular velocity
-	if (!CharacterInputs || CharacterInputs->OrientationIntent.IsNearlyZero())
+	if (!CharacterInputs || CharacterInputs->OrientationIntent.IsNearlyZero() || !bShouldApplyAngularVelocityToTarget)
 	{
 		IntendedOrientation_WorldSpace = StartingSyncState->GetOrientation_WorldSpace();
 
@@ -115,12 +115,15 @@ void UChaosFallingMode::GenerateMove_Implementation(const FMoverTickStartData& S
 	}
 	else
 	{
-		IntendedOrientation_WorldSpace = CharacterInputs->GetOrientationIntentDir_WorldSpace().ToOrientationRotator();
-
 		if (ShouldCharacterRemainUpright())
 		{
+			IntendedOrientation_WorldSpace = CharacterInputs->GetOrientationIntentDir_WorldSpace().ToOrientationRotator();
 			IntendedOrientation_WorldSpace = FRotator(FQuat::FindBetweenVectors(IntendedOrientation_WorldSpace.RotateVector(DefaultSimInputs->UpDir), DefaultSimInputs->UpDir) * IntendedOrientation_WorldSpace.Quaternion());
 		}
+		else
+		{
+			IntendedOrientation_WorldSpace = StartingSyncState->GetOrientation_WorldSpace();
+		}
 	}
 
 	FQuat WorldToGravityTransform = FQuat::FindBetweenNormals(FVector::UpVector, DefaultSimInputs->UpDir);
diff --git a/Engine/Plugins/Experimental/ChaosMover/Source/ChaosMover/Public/ChaosMover/Character/Modes/ChaosFallingMode.h b/Engine/Plugins/Experimental/ChaosMover/Source/ChaosMover/Public/ChaosMover/Character/Modes/ChaosFallingMode.h
--- a/Engine/Plugins/Experimental/ChaosMover/Source/ChaosMover/Public/ChaosMover/Character/Modes/ChaosFallingMode.h
+++ b/Engine/Plugins/Experimental/ChaosMover/Source/ChaosMover/Public/ChaosMover/Character/Modes/ChaosFallingMode.h
@@ -95,6 +95,12 @@ public:
 	UPROPERTY(Category = Mover, EditAnywhere, BlueprintReadWrite)
 	bool bAlwaysClampRequestedVelocityToMaxSpeed = false;
 
+	/**
+	 * Whether the character can apply angular velocity to respond to stick input
+	 */
+	UPROPERTY(Category = Mover, EditAnywhere, BlueprintReadWrite)
+	bool bShouldApplyAngularVelocityToTarget = true;
+
 protected:
 
 	/*
