[Chaos Mover] Add option to apply angular velocity to target in walking mode and improved angular velocity calculation in falling mode

#Synced-CL 48762495

diff --git a/Engine/Plugins/Experimental/ChaosMover/Source/ChaosMover/Private/Character/Modes/ChaosFallingMode.cpp b/Engine/Plugins/Experimental/ChaosMover/Source/ChaosMover/Private/Character/Modes/ChaosFallingMode.cpp
--- a/Engine/Plugins/Experimental/ChaosMover/Source/ChaosMover/Private/Character/Modes/ChaosFallingMode.cpp
+++ b/Engine/Plugins/Experimental/ChaosMover/Source/ChaosMover/Private/Character/Modes/ChaosFallingMode.cpp
@@ -208,11 +208,22 @@ void UChaosFallingMode::SimulationTick_Implementation(const FSimulationTickParam
 
 	OutputState.MovementEndState.RemainingMs = 0.0f;
 	OutputState.MovementEndState.NextModeName = Params.StartState.SyncState.MovementMode;
+
+	FVector TargetAngularVelocity = StartingSyncState->GetAngularVelocityDegrees_WorldSpace();
+	if (DeltaSeconds > UE_SMALL_NUMBER)
+	{
+		const FQuat InitialQuat = StartingSyncState->GetOrientation_WorldSpace().Quaternion();
+		FQuat TgtQuat = TargetOrient.Quaternion();
+		TgtQuat.EnforceShortestArcWith(InitialQuat);
+		const FQuat QuatRotation = TgtQuat * InitialQuat.Inverse();
+		const FVector AngularDisplacement = QuatRotation.ToRotationVector();
+		TargetAngularVelocity = FMath::RadiansToDegrees(AngularDisplacement / DeltaSeconds);
+	}
 	
 	OutputSyncState.MoveDirectionIntent = ProposedMove.bHasDirIntent ? ProposedMove.DirectionIntent : FVector::ZeroVector;
 	OutputSyncState.SetTransforms_WorldSpace(
 		TargetPos,
 		TargetOrient,
 		TargetVel,
-		ProposedMove.AngularVelocityDegrees);
+		TargetAngularVelocity);
 }
\ No newline at end of file
diff --git a/Engine/Plugins/Experimental/ChaosMover/Source/ChaosMover/Private/Character/Modes/ChaosWalkingMode.cpp b/Engine/Plugins/Experimental/ChaosMover/Source/ChaosMover/Private/Character/Modes/ChaosWalkingMode.cpp
--- a/Engine/Plugins/Experimental/ChaosMover/Source/ChaosMover/Private/Character/Modes/ChaosWalkingMode.cpp
+++ b/Engine/Plugins/Experimental/ChaosMover/Source/ChaosMover/Private/Character/Modes/ChaosWalkingMode.cpp
@@ -321,13 +321,27 @@ void UChaosWalkingMode::SimulationTick_Implementation(const FSimulationTickParam
 		// Target orientation
 		// This is always applied regardless of whether the character is supported
 		const FRotator TargetOrientation = UMovementUtils::ApplyAngularVelocityToRotator(StartingSyncState->GetOrientation_WorldSpace(), ProposedMove.AngularVelocityDegrees, DeltaSeconds);
+		FVector TargetAngularVelocity = StartingSyncState->GetAngularVelocityDegrees_WorldSpace();
+
+		if (ShouldCharacterRemainUpright() && bShouldApplyAngularVelocityToTarget)
+		{
+			if (DeltaSeconds > UE_SMALL_NUMBER)
+			{
+				const FQuat InitialQuat = StartingSyncState->GetOrientation_WorldSpace().Quaternion();
+				FQuat TgtQuat = TargetOrientation.Quaternion();
+				TgtQuat.EnforceShortestArcWith(InitialQuat);
+				const FQuat QuatRotation = TgtQuat * InitialQuat.Inverse();
+				const FVector AngularDisplacement = QuatRotation.ToRotationVector();
+				TargetAngularVelocity = FMath::RadiansToDegrees(AngularDisplacement / DeltaSeconds);
+			}
+		}
 
 		OutputSyncState.MoveDirectionIntent = ProposedMove.bHasDirIntent ? ProposedMove.DirectionIntent : FVector::ZeroVector;
 		OutputSyncState.SetTransforms_WorldSpace(
 			TargetPosition,
 			TargetOrientation,
 			TargetVelocity,
-			StartingSyncState->GetAngularVelocityDegrees_WorldSpace());
+			TargetAngularVelocity);
 	}
 	
 	OutputState.MovementEndState.RemainingMs = 0.0f;
diff --git a/Engine/Plugins/Experimental/ChaosMover/Source/ChaosMover/Public/ChaosMover/Character/Modes/ChaosWalkingMode.h b/Engine/Plugins/Experimental/ChaosMover/Source/ChaosMover/Public/ChaosMover/Character/Modes/ChaosWalkingMode.h
--- a/Engine/Plugins/Experimental/ChaosMover/Source/ChaosMover/Public/ChaosMover/Character/Modes/ChaosWalkingMode.h
+++ b/Engine/Plugins/Experimental/ChaosMover/Source/ChaosMover/Public/ChaosMover/Character/Modes/ChaosWalkingMode.h
@@ -57,6 +57,12 @@ public:
 	UPROPERTY(EditAnywhere, Category = "Movement Settings")
 	uint8 bMaintainHorizontalGroundVelocity : 1 = 0;
 
+	/**
+	 * Whether the character can apply angular velocity to respond to stick input
+	 */
+	UPROPERTY(Category = Mover, EditAnywhere, BlueprintReadWrite)
+	bool bShouldApplyAngularVelocityToTarget = true;
+
 	UE_API virtual void UpdateConstraintSettings(Chaos::FCharacterGroundConstraintSettings& ConstraintSettings) const override;
 
 	UE_API virtual void GenerateMove_Implementation(const FMoverTickStartData& StartState, const FMoverTimeStep& TimeStep, FProposedMove& OutProposedMove) const override;
