[ChaosMover] BUGFIX - CurrentMode variable in ChaosStanceModifier was not Serialized.

diff --git a/Engine/Plugins/Experimental/ChaosMover/Source/ChaosMover/Private/Character/Modifiers/ChaosStanceModifier.cpp b/Engine/Plugins/Experimental/ChaosMover/Source/ChaosMover/Private/Character/Modifiers/ChaosStanceModifier.cpp
--- a/Engine/Plugins/Experimental/ChaosMover/Source/ChaosMover/Private/Character/Modifiers/ChaosStanceModifier.cpp
+++ b/Engine/Plugins/Experimental/ChaosMover/Source/ChaosMover/Private/Character/Modifiers/ChaosStanceModifier.cpp
@@ -261,6 +261,7 @@ void FChaosStanceModifier::NetSerialize(FArchive& Ar)
 	Ar << DefaultCapsuleGroundClearance;
 	Ar << AccelerationOverride;
 	Ar << MaxSpeedOverride;
+	Ar << CurrentModeName;
 }
 
 UScriptStruct* FChaosStanceModifier::GetScriptStruct() const
