Only setting the FloorResult if the character is grounded, updating the MovementBase in post simulation tick
#rnx
#Synced-CL 48843018

diff --git a/Engine/Plugins/Experimental/ChaosMover/Source/ChaosMover/Private/ChaosMoverSimulation.cpp b/Engine/Plugins/Experimental/ChaosMover/Source/ChaosMover/Private/ChaosMoverSimulation.cpp
--- a/Engine/Plugins/Experimental/ChaosMover/Source/ChaosMover/Private/ChaosMoverSimulation.cpp
+++ b/Engine/Plugins/Experimental/ChaosMover/Source/ChaosMover/Private/ChaosMoverSimulation.cpp
@@ -486,7 +486,13 @@ void UChaosMoverSimulation::PostSimulationTickCharacter(const IChaosCharacterMov
 	check(Blackboard.Get());
 	FFloorCheckResult FloorResult;
 	bool FoundLastFloorResult = Blackboard->TryGet(CommonBlackboard::LastFloorResult, FloorResult);
-	if (FoundLastFloorResult && FloorResult.bBlockingHit)
+
+	bool bIsInGroundMode = false;
+	if (TStrongObjectPtr<const UBaseMovementMode> Mode = GetCurrentMovementMode().Pin())
+	{
+		bIsInGroundMode = Mode->HasGameplayTag(Mover_IsOnGround, true);
+	}
+	if (FoundLastFloorResult && FloorResult.bBlockingHit && bIsInGroundMode)
 	{
 		const FVector WorldSpaceDirectionIntent = PostSimDefaultSyncState->GetIntent_WorldSpace();
 		PostSimDefaultSyncState->SetTransforms_WorldSpace(
