ChaosMover: Tracing network physics component information to CVD. This was specifically to understand input decay an debug lag scaling.

#Synced-CL 48814712

diff --git a/Engine/Plugins/Experimental/ChaosMover/Source/ChaosMover/Private/ChaosMoverSimulation.cpp b/Engine/Plugins/Experimental/ChaosMover/Source/ChaosMover/Private/ChaosMoverSimulation.cpp
--- a/Engine/Plugins/Experimental/ChaosMover/Source/ChaosMover/Private/ChaosMoverSimulation.cpp
+++ b/Engine/Plugins/Experimental/ChaosMover/Source/ChaosMover/Private/ChaosMoverSimulation.cpp
@@ -894,13 +894,21 @@ void UChaosMoverSimulation::TraceMoverData(const FMoverTimeStep& TimeStep, const
 #if WITH_CHAOS_VISUAL_DEBUGGER
 	if (FChaosVisualDebuggerTrace::IsTracing())
 	{
-		static FName NAME_LocalSimImput("LocalSimImput");
-		static FName NAME_InternalSimData("InternalSimData");
-		static FName NAME_DebugSimData("DebugSimData");
-
+		// Trace time step info
 		FChaosMoverTimeStepDebugData& TimeStepDebugData = DebugSimData.FindOrAddMutableDataByType<FChaosMoverTimeStepDebugData>();
 		TimeStepDebugData.SetTimeStep(TimeStep);
 
+		// Trace network physics info
+		if (const FChaosMoverSimulationDefaultInputs* SimInputs = LocalSimInput.FindDataByType<FChaosMoverSimulationDefaultInputs>())
+		{
+			FNetworkPhysicsDebugData& NetworkPhysicsDebugData = DebugSimData.FindOrAddMutableDataByType<FNetworkPhysicsDebugData>();
+			NetworkPhysicsDebugData.Set(SimInputs->AsyncNetworkPhysicsComponent);
+		}
+
+		// Prepare tracing of LocalSimInput, InternalSimData and DebugSimData
+		static FName NAME_LocalSimImput("LocalSimImput");
+		static FName NAME_InternalSimData("InternalSimData");
+		static FName NAME_DebugSimData("DebugSimData");
 		UE::MoverUtils::NamedDataCollections LocalSimDataCollections (
 			{				
 				{ NAME_LocalSimImput, &LocalSimInput},
diff --git a/Engine/Plugins/Experimental/ChaosMover/Source/ChaosMover/Private/ChaosMoverSimulationTypes.cpp b/Engine/Plugins/Experimental/ChaosMover/Source/ChaosMover/Private/ChaosMoverSimulationTypes.cpp
--- a/Engine/Plugins/Experimental/ChaosMover/Source/ChaosMover/Private/ChaosMoverSimulationTypes.cpp
+++ b/Engine/Plugins/Experimental/ChaosMover/Source/ChaosMover/Private/ChaosMoverSimulationTypes.cpp
@@ -8,6 +8,7 @@
 #include "ChaosMover/ChaosMoverConsoleVariables.h"
 #include "MoveLibrary/BasedMovementUtils.h"
 #include "Serialization/ArchiveCrc32.h"
+#include "Physics/NetworkPhysicsComponent.h"
 
 #include UE_INLINE_GENERATED_CPP_BY_NAME(ChaosMoverSimulationTypes)
 
@@ -119,6 +120,7 @@ void FChaosMoverTimeStepDebugData::SetTimeStep(const FMoverTimeStep& InTimeStep)
 {
 	TimeStep = InTimeStep;
 	bIsResimulating = InTimeStep.bIsResimulating;
+	bIsFirstResimFrame = InTimeStep.bIsFirstResimFrame;
 }
 
 FMoverDataStructBase* FChaosMoverTimeStepDebugData::Clone() const
@@ -131,6 +133,31 @@ UScriptStruct* FChaosMoverTimeStepDebugData::GetScriptStruct() const
 	return StaticStruct();
 }
 
+FMoverDataStructBase* FNetworkPhysicsDebugData::Clone() const
+{
+	return new FNetworkPhysicsDebugData(*this);
+}
+
+void FNetworkPhysicsDebugData::Set(const FAsyncNetworkPhysicsComponent* NetworkPhysicsComponent)
+{
+	if (NetworkPhysicsComponent)
+	{
+		bIsLocallyControlled = NetworkPhysicsComponent->IsLocallyControlled();
+		PhysicsReplicationMode = NetworkPhysicsComponent->GetPhysicsReplicationMode();
+		NetworkPhysicsTickOffset = NetworkPhysicsComponent->GetNetworkPhysicsTickOffset();
+		LastReceivedServerFrame = NetworkPhysicsComponent->GetLastReceivedServerFrame();
+		LatestReceivedTargetStateFrame = NetworkPhysicsComponent->GetLatestReceivedTargetStateFrame();
+		ForwardPredictionTime = NetworkPhysicsComponent->GetForwardPredictionTime();
+		CurrentSimProxyInputDecayAtRuntime = NetworkPhysicsComponent->GetCurrentSimProxyInputDecayAtRuntime();
+		CurrentInputDecay = NetworkPhysicsComponent->GetCurrentInputDecay();
+	}
+}
+
+UScriptStruct* FNetworkPhysicsDebugData::GetScriptStruct() const
+{
+	return StaticStruct();
+}
+
 void FChaosNetInstantMovementEffectsQueue::Add(const FScheduledInstantMovementEffect& InScheduledEffect, int32 IssuanceServerFrame, uint8 UniqueID)
 {
 	TCheckedObjPtr<FInstantMovementEffect> InstantMovementEffect = InScheduledEffect.Effect.Get();
@@ -243,7 +270,7 @@ void FChaosNetInstantMovementEffectsQueue::ToString(FAnsiStringBuilderBase& Out)
 	for (int32 EffectIndex = 0; EffectIndex < Effects.Num(); ++EffectIndex)
 	{
 		const FChaosNetInstantMovementEffect& ScheduledEffect = Effects[EffectIndex];
-		Out.Appendf("Effect Index %d: ServerFrame %d, ", EffectIndex, ScheduledEffect.ExecutionServerFrame, ScheduledEffect.Effect.IsValid() ? *ScheduledEffect.Effect.Get().ToSimpleString() : TEXT("INVALID INSTANCED STRUCT EFFECT") );
+		Out.Appendf("Effect Index %d: ServerFrame %d, %s", EffectIndex, ScheduledEffect.ExecutionServerFrame, ScheduledEffect.Effect.IsValid() ? *ScheduledEffect.Effect.Get().ToSimpleString() : TEXT("INVALID INSTANCED STRUCT EFFECT") );
 		Out.AppendChar('\n');
 	}
 	Out.Appendf("--------------------------------------------------------------------------------\n");
diff --git a/Engine/Plugins/Experimental/ChaosMover/Source/ChaosMover/Public/ChaosMover/ChaosMoverSimulationTypes.h b/Engine/Plugins/Experimental/ChaosMover/Source/ChaosMover/Public/ChaosMover/ChaosMoverSimulationTypes.h
--- a/Engine/Plugins/Experimental/ChaosMover/Source/ChaosMover/Public/ChaosMover/ChaosMoverSimulationTypes.h
+++ b/Engine/Plugins/Experimental/ChaosMover/Source/ChaosMover/Public/ChaosMover/ChaosMoverSimulationTypes.h
@@ -209,16 +209,56 @@ struct FChaosMoverTimeStepDebugData : public FMoverDataStructBase
 	virtual FMoverDataStructBase* Clone() const override;
 	virtual UScriptStruct* GetScriptStruct() const override;
 
-	void SetTimeStep(const FMoverTimeStep& InTimeStep);
+	CHAOSMOVER_API void SetTimeStep(const FMoverTimeStep& InTimeStep);
 
-	// This is so CVD can display TimeStep.bIsResimulating properly, which does not exist in FMoverTimeStep in a way that will show up by default
+	// This is so CVD can display TimeStep.bIsResimulating properly, which does not exist in FMoverTimeStep in a way that will show up by default (bitflag)
 	UPROPERTY(VisibleAnywhere, Category = "Mover Info")
 	bool bIsResimulating = false;
 
+	// This is so CVD can display TimeStep.bIsFirstResimFrame properly, which does not exist in FMoverTimeStep in a way that will show up by default (bitflag)
+	UPROPERTY(VisibleAnywhere, Category = "Mover Info")
+	bool bIsFirstResimFrame = false;
+
 	UPROPERTY(VisibleAnywhere, Category = "Mover Info")
 	FMoverTimeStep TimeStep;
 };
 
+// Structure for tracing some network physics info to CVD
+USTRUCT()
+struct FNetworkPhysicsDebugData : public FMoverDataStructBase
+{
+	GENERATED_BODY()
+
+	virtual FMoverDataStructBase* Clone() const override;
+	virtual UScriptStruct* GetScriptStruct() const override;
+
+	CHAOSMOVER_API void Set(const FAsyncNetworkPhysicsComponent* NetworkPhysicsComponent);
+
+	UPROPERTY(VisibleAnywhere, Category = "Mover Info")
+	bool bIsLocallyControlled = false;
+
+	UPROPERTY(VisibleAnywhere, Category = "Mover Info")
+	EPhysicsReplicationMode PhysicsReplicationMode = EPhysicsReplicationMode::Default;
+
+	UPROPERTY(VisibleAnywhere, Category = "Mover Info")
+	int32 NetworkPhysicsTickOffset = INDEX_NONE;
+
+	UPROPERTY(VisibleAnywhere, Category = "Mover Info")
+	int32 LastReceivedServerFrame = INDEX_NONE;
+
+	UPROPERTY(VisibleAnywhere, Category = "Mover Info")
+	int32 LatestReceivedTargetStateFrame = INDEX_NONE;
+
+	UPROPERTY(VisibleAnywhere, Category = "Mover Info")
+	float ForwardPredictionTime = -1.0f;
+
+	UPROPERTY(VisibleAnywhere, Category = "Mover Info")
+	float CurrentSimProxyInputDecayAtRuntime = -1.0f;
+
+	UPROPERTY(VisibleAnywhere, Category = "Mover Info")
+	float CurrentInputDecay = -1.0f;
+};
+
 UINTERFACE(BlueprintType, MinimalAPI, meta = (CannotImplementInterfaceInBlueprint))
 class UChaosCharacterMovementModeInterface : public UInterface
 {
