[ChaosMover] Fixes for large velocities with bots
Preflight was done with another CL and had to be split to be submitted.
#rnx
#tests PIE,launchit
#Synced-CL 49124608

diff --git a/Engine/Plugins/Experimental/ChaosMover/Source/ChaosMover/Private/Character/Modes/ChaosFallingMode.cpp b/Engine/Plugins/Experimental/ChaosMover/Source/ChaosMover/Private/Character/Modes/ChaosFallingMode.cpp
--- a/Engine/Plugins/Experimental/ChaosMover/Source/ChaosMover/Private/Character/Modes/ChaosFallingMode.cpp
+++ b/Engine/Plugins/Experimental/ChaosMover/Source/ChaosMover/Private/Character/Modes/ChaosFallingMode.cpp
@@ -9,6 +9,7 @@
 #include "ChaosMover/Character/Transitions/ChaosCharacterLandingCheck.h"
 #include "ChaosMover/Utilities/ChaosMoverQueryUtils.h"
 #include "MoveLibrary/AirMovementUtils.h"
+#include "MoveLibrary/ConstrainedMoveUtils.h"
 #include "MoveLibrary/FloorQueryUtils.h"
 #include "MoveLibrary/MovementUtils.h"
 #include "MoveLibrary/WaterMovementUtils.h"
@@ -76,11 +77,11 @@ void UChaosFallingMode::GenerateMove_Implementation(const FMoverTickStartData& S
 	const FChaosMoverRequestedMoveInputs* RequestedMove = StartState.InputCmd.InputCollection.FindDataByType<FChaosMoverRequestedMoveInputs>();
 	if (RequestedMove)
 	{
-		Params.MoveInputType = EMoveInputType::Velocity;
-		// It's a bit confusing but bForceMaxSpeed is ill named and means that the requested velocity will NOT be clamped to max speed
-		// bAlwaysClampRequestedVelocityToMaxSpeed takes precedence over bForceMaxSpeed
-		const bool bClampToMaxSpeed = bAlwaysClampRequestedVelocityToMaxSpeed || !RequestedMove->bForceMaxSpeed;
-		Params.MoveInput = bClampToMaxSpeed ? RequestedMove->RequestedVelocity.GetClampedToMaxSize(GetMaxSpeed()) : RequestedMove->RequestedVelocity;
+		Params.MoveInputType = EMoveInputType::DirectionalIntent;
+		const float MaxSpeed = GetMaxSpeed();
+		const FVector DirInput = MaxSpeed > 0 ? (RequestedMove->RequestedVelocity / MaxSpeed).GetSafeNormal() : FVector::ZeroVector;
+		constexpr bool bMaintainInputMagnitude = true;
+		Params.MoveInput = UPlanarConstraintUtils::ConstrainDirectionToPlane(FPlanarConstraint(), DirInput, bMaintainInputMagnitude);
 	}
 	else
 	{
