[Chaos Mover] Move ground query back to generate move by default to save additional query cost in walking mode

#Synced-CL 49327620

diff --git a/Engine/Plugins/Experimental/ChaosMover/Source/ChaosMover/Private/ChaosMoverConsoleVariables.cpp b/Engine/Plugins/Experimental/ChaosMover/Source/ChaosMover/Private/ChaosMoverConsoleVariables.cpp
--- a/Engine/Plugins/Experimental/ChaosMover/Source/ChaosMover/Private/ChaosMoverConsoleVariables.cpp
+++ b/Engine/Plugins/Experimental/ChaosMover/Source/ChaosMover/Private/ChaosMoverConsoleVariables.cpp
@@ -51,6 +51,9 @@ namespace UE::ChaosMover
 		FAutoConsoleVariableRef CVarChaosMoverFallingCheckRelativeSpeedLimit(TEXT("ChaosMover.FallingCheckRelativeSpeedLimit"),
 			FallingCheckRelativeSpeedLimit, TEXT("Above this limit the falling check will consider the character to be lifting off the surface.\n"));
 
+		bool bEnablePreSimGroundCheck = false;
+		FAutoConsoleVariableRef CVarChaosMoverEnablePreSimGroundCheck(TEXT("ChaosMover.EnablePreSimGroundCheck"),
+			bEnablePreSimGroundCheck, TEXT(""));
 
 		static FAutoConsoleCommandWithWorld TeleportOnlyClientCommand
 		(
diff --git a/Engine/Plugins/Experimental/ChaosMover/Source/ChaosMover/Private/ChaosMoverSimulation.cpp b/Engine/Plugins/Experimental/ChaosMover/Source/ChaosMover/Private/ChaosMoverSimulation.cpp
--- a/Engine/Plugins/Experimental/ChaosMover/Source/ChaosMover/Private/ChaosMoverSimulation.cpp
+++ b/Engine/Plugins/Experimental/ChaosMover/Source/ChaosMover/Private/ChaosMoverSimulation.cpp
@@ -379,8 +379,11 @@ void UChaosMoverSimulation::PreSimulationTickCharacter(const IChaosCharacterMove
 		}
 	}
 
-	// Update the floor before we run any update.
-	CharacterMode.UpdateCurrentFloor(TimeStep);
+	if (UE::ChaosMover::CVars::bEnablePreSimGroundCheck)
+	{
+		// Update the floor before we run any update.
+		CharacterMode.UpdateCurrentFloor(TimeStep);
+	}
 }
 
 void UChaosMoverSimulation::OnSimulationRollback(const FMoverTimeStep& NewTimeStep, const FMoverSyncState& PrevSyncState)
diff --git a/Engine/Plugins/Experimental/ChaosMover/Source/ChaosMover/Private/Character/Modes/ChaosFallingMode.cpp b/Engine/Plugins/Experimental/ChaosMover/Source/ChaosMover/Private/Character/Modes/ChaosFallingMode.cpp
--- a/Engine/Plugins/Experimental/ChaosMover/Source/ChaosMover/Private/Character/Modes/ChaosFallingMode.cpp
+++ b/Engine/Plugins/Experimental/ChaosMover/Source/ChaosMover/Private/Character/Modes/ChaosFallingMode.cpp
@@ -2,6 +2,7 @@
 
 #include "ChaosMover/Character/Modes/ChaosFallingMode.h"
 
+#include "ChaosMover/ChaosMoverConsoleVariables.h"
 #include "ChaosMover/ChaosMoverLog.h"
 #include "ChaosMover/ChaosMoverSimulation.h"
 #include "ChaosMover/Character/ChaosCharacterInputs.h"
@@ -174,6 +175,37 @@ void UChaosFallingMode::GenerateMove_Implementation(const FMoverTickStartData& S
 	{
 		UMovementUtils::SetGravityVerticalComponent(OutProposedMove.LinearVelocity, VelocityWithGravity.Dot(DefaultSimInputs->UpDir), DefaultSimInputs->UpDir);
 	}
+
+
+	if (!UE::ChaosMover::CVars::bEnablePreSimGroundCheck)
+	{
+		UMoverBlackboard* SimBlackboard = Simulation->GetBlackboard_Mutable();
+		if (SimBlackboard)
+		{
+			FFloorCheckResult FloorResult;
+			FWaterCheckResult WaterResult;
+
+			UE::ChaosMover::Utils::FFloorSweepParams SweepParams{
+				.ResponseParams = DefaultSimInputs->CollisionResponseParams,
+				.QueryParams = DefaultSimInputs->CollisionQueryParams,
+				.Location = StartingSyncState->GetLocation_WorldSpace(),
+				.DeltaPos = OutProposedMove.LinearVelocity * DeltaSeconds,
+				.UpDir = DefaultSimInputs->UpDir,
+				.World = DefaultSimInputs->World,
+				.QueryDistance = 1.5f * GetTargetHeight(),
+				.QueryRadius = FMath::Min(GetGroundQueryRadius(), FMath::Max(DefaultSimInputs->PawnCollisionRadius - 5.0f, 0.0f)),
+				.MaxWalkSlopeCosine = GetMaxWalkSlopeCosine(),
+				.TargetHeight = GetTargetHeight(),
+				.CollisionChannel = DefaultSimInputs->CollisionChannel
+			};
+
+			UE::ChaosMover::Utils::FloorSweep_Internal(SweepParams, FloorResult, WaterResult);
+
+			SimBlackboard->Set(CommonBlackboard::LastFloorResult, FloorResult);
+			SimBlackboard->Set(CommonBlackboard::LastWaterResult, WaterResult);
+			SimBlackboard->Set(UE::ChaosMover::Blackboard::GroundDynamicsInfo, UE::ChaosMover::FGroundDynamicsInfo(FloorResult));
+		}
+	}
 }
 
 void UChaosFallingMode::SimulationTick_Implementation(const FSimulationTickParams& Params, FMoverTickEndData& OutputState)
diff --git a/Engine/Plugins/Experimental/ChaosMover/Source/ChaosMover/Private/Character/Modes/ChaosFlyingMode.cpp b/Engine/Plugins/Experimental/ChaosMover/Source/ChaosMover/Private/Character/Modes/ChaosFlyingMode.cpp
--- a/Engine/Plugins/Experimental/ChaosMover/Source/ChaosMover/Private/Character/Modes/ChaosFlyingMode.cpp
+++ b/Engine/Plugins/Experimental/ChaosMover/Source/ChaosMover/Private/Character/Modes/ChaosFlyingMode.cpp
@@ -91,11 +91,11 @@ void UChaosFlyingMode::GenerateMove_Implementation(const FMoverTickStartData& St
 	OutProposedMove = UAirMovementUtils::ComputeControlledFreeMove(Params);
 
 	// Don't do floor checks in flying mode so just clear any previous results from the blackboard
-	UMoverBlackboard* SimBlackboard = Simulation->GetBlackboard_Mutable();
-	if (SimBlackboard)
+	if (UMoverBlackboard* SimBlackboard = Simulation->GetBlackboard_Mutable())
 	{
 		SimBlackboard->Invalidate(CommonBlackboard::LastFloorResult);
 		SimBlackboard->Invalidate(CommonBlackboard::LastWaterResult);
+		SimBlackboard->Invalidate(UE::ChaosMover::Blackboard::GroundDynamicsInfo);
 	}
 }
 
@@ -140,14 +140,8 @@ void UChaosFlyingMode::SimulationTick_Implementation(const FSimulationTickParams
 		ProposedMove.AngularVelocityDegrees);
 }
 
-void UChaosFlyingMode::UpdateCurrentFloor(const FMoverTimeStep&) const
+void UChaosFlyingMode::UpdateCurrentFloor(const FMoverTimeStep& TimeStep) const
 {
-	if (!Simulation)
-	{
-		UE_LOG(LogChaosMover, Warning, TEXT("No Simulation set on ChaosFlyingMode"));
-		return;
-	}
-
 	if (UMoverBlackboard* SimBlackboard = Simulation->GetBlackboard_Mutable())
 	{
 		SimBlackboard->Invalidate(CommonBlackboard::LastFloorResult);
diff --git a/Engine/Plugins/Experimental/ChaosMover/Source/ChaosMover/Private/Character/Modes/ChaosSwimmingMode.cpp b/Engine/Plugins/Experimental/ChaosMover/Source/ChaosMover/Private/Character/Modes/ChaosSwimmingMode.cpp
--- a/Engine/Plugins/Experimental/ChaosMover/Source/ChaosMover/Private/Character/Modes/ChaosSwimmingMode.cpp
+++ b/Engine/Plugins/Experimental/ChaosMover/Source/ChaosMover/Private/Character/Modes/ChaosSwimmingMode.cpp
@@ -2,6 +2,7 @@
 
 #include "ChaosMover/Character/Modes/ChaosSwimmingMode.h"
 
+#include "ChaosMover/ChaosMoverConsoleVariables.h"
 #include "ChaosMover/ChaosMoverLog.h"
 #include "ChaosMover/ChaosMoverSimulation.h"
 #include "ChaosMover/Character/Settings/SharedChaosCharacterMovementSettings.h"
@@ -56,28 +57,40 @@ void UChaosSwimmingMode::GenerateMove_Implementation(const FMoverTickStartData&
 	FVector Velocity = StartingSyncState->GetVelocity_WorldSpace();
 	
 	UMoverBlackboard* SimBlackboard = Simulation->GetBlackboard_Mutable();
-	SimBlackboard->Invalidate(CommonBlackboard::LastFloorResult);
-	SimBlackboard->Invalidate(CommonBlackboard::LastWaterResult);
 	FFloorCheckResult FloorResult;
 	FWaterCheckResult WaterResult;
-
-	const float QueryDistance = 2.0f * HalfHeight;
-	FVector DeltaPos = Velocity * DeltaSeconds;
-	UE::ChaosMover::Utils::FFloorSweepParams SweepParams{
-		.ResponseParams = DefaultSimInputs->CollisionResponseParams,
-		.QueryParams = DefaultSimInputs->CollisionQueryParams,
-		.Location = StartingSyncState->GetLocation_WorldSpace(),
-		.DeltaPos = DeltaPos,
-		.UpDir = DefaultSimInputs->UpDir,
-		.World = DefaultSimInputs->World,
-		.QueryDistance = QueryDistance,
-		.QueryRadius = FMath::Min(GetGroundQueryRadius(), FMath::Max(DefaultSimInputs->PawnCollisionRadius - 5.0f, 0.0f)),
-		.MaxWalkSlopeCosine = GetMaxWalkSlopeCosine(),
-		.TargetHeight = GetTargetHeight(),
-		.CollisionChannel = DefaultSimInputs->CollisionChannel
-	};
-
-	UE::ChaosMover::Utils::FloorSweep_Internal(SweepParams, FloorResult, WaterResult);
+	
+	if (UE::ChaosMover::CVars::bEnablePreSimGroundCheck)
+	{
+		SimBlackboard->TryGet(CommonBlackboard::LastFloorResult, FloorResult);
+		SimBlackboard->TryGet(CommonBlackboard::LastWaterResult, WaterResult);
+	}
+	else
+	{
+		SimBlackboard->Invalidate(CommonBlackboard::LastFloorResult);
+		SimBlackboard->Invalidate(CommonBlackboard::LastWaterResult);
+
+		const float QueryDistance = 2.0f * HalfHeight;
+		FVector DeltaPos = Velocity * DeltaSeconds;
+		UE::ChaosMover::Utils::FFloorSweepParams SweepParams{
+			.ResponseParams = DefaultSimInputs->CollisionResponseParams,
+			.QueryParams = DefaultSimInputs->CollisionQueryParams,
+			.Location = StartingSyncState->GetLocation_WorldSpace(),
+			.DeltaPos = DeltaPos,
+			.UpDir = DefaultSimInputs->UpDir,
+			.World = DefaultSimInputs->World,
+			.QueryDistance = QueryDistance,
+			.QueryRadius = FMath::Min(GetGroundQueryRadius(), FMath::Max(DefaultSimInputs->PawnCollisionRadius - 5.0f, 0.0f)),
+			.MaxWalkSlopeCosine = GetMaxWalkSlopeCosine(),
+			.TargetHeight = GetTargetHeight(),
+			.CollisionChannel = DefaultSimInputs->CollisionChannel
+		};
+
+		UE::ChaosMover::Utils::FloorSweep_Internal(SweepParams, FloorResult, WaterResult);
+
+		SimBlackboard->Set(CommonBlackboard::LastFloorResult, FloorResult);
+		SimBlackboard->Set(UE::ChaosMover::Blackboard::GroundDynamicsInfo, UE::ChaosMover::FGroundDynamicsInfo(FloorResult));
+	}
 
 	if (WaterResult.IsSwimmableVolume())
 	{
@@ -92,7 +105,6 @@ void UChaosSwimmingMode::GenerateMove_Implementation(const FMoverTickStartData&
 		UWaterMovementUtils::UpdateWaterSplineData(SplineParams, WaterResult);
 	}
 
-	SimBlackboard->Set(CommonBlackboard::LastFloorResult, FloorResult);
 	SimBlackboard->Set(CommonBlackboard::LastWaterResult, WaterResult);
 	
 	// Buoyancy Bobbing
diff --git a/Engine/Plugins/Experimental/ChaosMover/Source/ChaosMover/Private/Character/Modes/ChaosWalkingMode.cpp b/Engine/Plugins/Experimental/ChaosMover/Source/ChaosMover/Private/Character/Modes/ChaosWalkingMode.cpp
--- a/Engine/Plugins/Experimental/ChaosMover/Source/ChaosMover/Private/Character/Modes/ChaosWalkingMode.cpp
+++ b/Engine/Plugins/Experimental/ChaosMover/Source/ChaosMover/Private/Character/Modes/ChaosWalkingMode.cpp
@@ -192,8 +192,23 @@ void UChaosWalkingMode::GenerateMove_Implementation(const FMoverTickStartData& S
 		{
 			OutProposedMove.LinearVelocity = FVector::ZeroVector;
 		}
-
 	}
+
+	if (SimBlackboard)
+	{
+		// Update the floor result and check the proposed move to prevent movement onto unwalkable surfaces
+		FVector OutDeltaPos = FVector::ZeroVector;
+		FFloorCheckResult FloorResult;
+		FWaterCheckResult WaterResult;
+		GetFloorAndCheckMovement(*StartingSyncState, OutProposedMove, *DefaultSimInputs, DeltaSeconds, FloorResult, WaterResult, OutDeltaPos);
+
+		OutProposedMove.LinearVelocity = (DeltaSeconds > 0.0f) ? (OutDeltaPos / DeltaSeconds) : FVector::ZeroVector;
+
+		GroundDynamicsInfo = UE::ChaosMover::FGroundDynamicsInfo(FloorResult);
+		SimBlackboard->Set(CommonBlackboard::LastFloorResult, FloorResult);
+		SimBlackboard->Set(UE::ChaosMover::Blackboard::GroundDynamicsInfo, GroundDynamicsInfo);
+		SimBlackboard->Set(CommonBlackboard::LastWaterResult, WaterResult);
+	};
 }
 
 void UChaosWalkingMode::SimulationTick_Implementation(const FSimulationTickParams& Params, FMoverTickEndData& OutputState)
@@ -223,25 +238,17 @@ void UChaosWalkingMode::SimulationTick_Implementation(const FSimulationTickParam
 
 	// Update the ground based on the proposed move
 	const float DeltaSeconds = Params.TimeStep.StepMs * 0.001f;
-	FVector DeltaPos;
-	FFloorCheckResult FloorResult;
-	FWaterCheckResult WaterResult;
-	GetFloorAndCheckMovement(OutputSyncState, ProposedMove, *DefaultSimInputs, DeltaSeconds, FloorResult, WaterResult, DeltaPos);
-
-	ProposedMove.LinearVelocity = (DeltaSeconds > 0.0f) ? (DeltaPos / DeltaSeconds) : FVector::ZeroVector;
-
-	UE::ChaosMover::FGroundDynamicsInfo GroundDynamicsInfo(FloorResult);
+	const FVector UpDirection = DefaultSimInputs->UpDir;
+	FVector GroundNormal = UpDirection;
 
-	if (UMoverBlackboard* SimBlackboard = Simulation->GetBlackboard_Mutable())
+	FFloorCheckResult FloorResult;
+	UE::ChaosMover::FGroundDynamicsInfo GroundDynamicsInfo;
+	if (const UMoverBlackboard* SimBlackboard = Simulation->GetBlackboard())
 	{
-		SimBlackboard->Set(CommonBlackboard::LastFloorResult, FloorResult);
-		SimBlackboard->Set(CommonBlackboard::LastWaterResult, WaterResult);
-		SimBlackboard->Set(UE::ChaosMover::Blackboard::GroundDynamicsInfo, GroundDynamicsInfo);
+		SimBlackboard->TryGet(CommonBlackboard::LastFloorResult, FloorResult);
+		SimBlackboard->TryGet(UE::ChaosMover::Blackboard::GroundDynamicsInfo, GroundDynamicsInfo);
 	}
 
-	const FVector UpDirection = DefaultSimInputs->UpDir;
-	FVector GroundNormal = UpDirection;
-
 	if (FloorResult.IsWalkableFloor())
 	{
 		GroundNormal = FloorResult.HitResult.ImpactNormal;
diff --git a/Engine/Plugins/Experimental/ChaosMover/Source/ChaosMover/Private/Character/Transitions/ChaosCharacterFallingCheck.cpp b/Engine/Plugins/Experimental/ChaosMover/Source/ChaosMover/Private/Character/Transitions/ChaosCharacterFallingCheck.cpp
--- a/Engine/Plugins/Experimental/ChaosMover/Source/ChaosMover/Private/Character/Transitions/ChaosCharacterFallingCheck.cpp
+++ b/Engine/Plugins/Experimental/ChaosMover/Source/ChaosMover/Private/Character/Transitions/ChaosCharacterFallingCheck.cpp
@@ -77,8 +77,7 @@ FTransitionEvalResult UChaosCharacterFallingCheck::Evaluate_Implementation(const
 				const float DeltaSeconds = Params.TimeStep.StepMs * 0.001f;
 				const FVector ProjectedVelocity = StartingSyncState->GetVelocity_WorldSpace() + SimInputs->Gravity * DeltaSeconds;
 				FVector ProjectedGroundVelocity = GroundDynamicsInfo.LinearVelocity;
-				const Chaos::FPBDRigidParticleHandle* GroundParticle = UChaosGroundMovementUtils::GetRigidParticleHandleFromFloorResult_Internal(FloorResult);
-				if (GroundParticle && GroundParticle->IsDynamic() && GroundParticle->GravityEnabled())
+				if (GroundDynamicsInfo.bIsDynamic && GroundDynamicsInfo.bIsGravityEnabled)
 				{
 					// This might not be correct if different physics objects have different gravity but is saves having to go
 					// to the component to get the gravity on the physics volume.
diff --git a/Engine/Plugins/Experimental/ChaosMover/Source/ChaosMover/Public/ChaosMover/ChaosMoverConsoleVariables.h b/Engine/Plugins/Experimental/ChaosMover/Source/ChaosMover/Public/ChaosMover/ChaosMoverConsoleVariables.h
--- a/Engine/Plugins/Experimental/ChaosMover/Source/ChaosMover/Public/ChaosMover/ChaosMoverConsoleVariables.h
+++ b/Engine/Plugins/Experimental/ChaosMover/Source/ChaosMover/Public/ChaosMover/ChaosMoverConsoleVariables.h
@@ -17,5 +17,8 @@ namespace UE::ChaosMover
 		CHAOSMOVER_API extern int32 InstantMovementEffectIDHistorySize;
 		CHAOSMOVER_API extern int32 FrameDifferenceLeniencyForEventTimeComparison;
 		CHAOSMOVER_API extern float FallingCheckRelativeSpeedLimit;
+
+		// Temporary
+		CHAOSMOVER_API extern bool bEnablePreSimGroundCheck;
 	}
 };
