[Chaos Mover] Change default radial force limit in falling mode to zero

#Synced-CL 49363637

diff --git a/Engine/Plugins/Experimental/ChaosMover/Source/ChaosMover/Private/Character/Modes/ChaosFallingMode.cpp b/Engine/Plugins/Experimental/ChaosMover/Source/ChaosMover/Private/Character/Modes/ChaosFallingMode.cpp
--- a/Engine/Plugins/Experimental/ChaosMover/Source/ChaosMover/Private/Character/Modes/ChaosFallingMode.cpp
+++ b/Engine/Plugins/Experimental/ChaosMover/Source/ChaosMover/Private/Character/Modes/ChaosFallingMode.cpp
@@ -36,9 +36,7 @@ UChaosFallingMode::UChaosFallingMode(const FObjectInitializer& ObjectInitializer
 	GameplayTags.AddTag(Mover_IsFalling);
 	GameplayTags.AddTag(Mover_SkipAnimRootMotion);
 
-	// Note: When not on the ground the constraint will not apply a radial force. Allowing for a high radial force
-	// limit here helps maintain momentum on landing before transitioning to walking
-	RadialForceLimit = 300000.0f;
+	RadialForceLimit = 0.0f;
 	SwingTorqueLimit = 3000.0f;
 	TwistTorqueLimit = 0.0f;
 
