[ChaosMover] Crash on UChaosWalkingMode::SimulationTick_Implementation
#tech-notes: A pointer that may be invalid is validated
#rnx

diff --git a/Engine/Plugins/Experimental/ChaosMover/Source/ChaosMover/Private/ChaosMoverStateMachine.cpp b/Engine/Plugins/Experimental/ChaosMover/Source/ChaosMover/Private/ChaosMoverStateMachine.cpp
--- a/Engine/Plugins/Experimental/ChaosMover/Source/ChaosMover/Private/ChaosMoverStateMachine.cpp
+++ b/Engine/Plugins/Experimental/ChaosMover/Source/ChaosMover/Private/ChaosMoverStateMachine.cpp
@@ -456,82 +456,64 @@ namespace UE::ChaosMover
 			if (!CurrentModeName.IsNone() && Modes.Contains(CurrentModeName))
 			{
 				TStrongObjectPtr<UBaseMovementMode> CurrentMode = Modes[CurrentModeName].Pin();
-				FProposedMove CombinedMove;
-				bool bHasModeMoveContribution = false;
-
-				if (!CVars::bSkipGenerateMoveIfOverridden ||
-					!(bHasLayeredMoveContributions && CombinedLayeredMove.MixMode == EMoveMixMode::OverrideAll))
+				if (CurrentMode)
 				{
-					QUICK_SCOPE_CYCLE_COUNTER(MoverStateMachine_OnSimulationTick_GenerateMoveFromMode);
-					CurrentMode->GenerateMove(WorkingSubstepStartData, SubTimeStep, OUT CombinedMove);
-
-					bHasModeMoveContribution = true;
-				}
+					FProposedMove CombinedMove;
+					bool bHasModeMoveContribution = false;
 
-				if (bHasModeMoveContribution && bHasLayeredMoveContributions)
-				{
-					FVector UpDir = FVector::UpVector;
-					if (const FChaosMoverSimulationDefaultInputs* DefaultSimInputs = Simulation->GetLocalSimInput().FindDataByType<FChaosMoverSimulationDefaultInputs>())
+					if (!CVars::bSkipGenerateMoveIfOverridden ||
+						!(bHasLayeredMoveContributions && CombinedLayeredMove.MixMode == EMoveMixMode::OverrideAll))
 					{
-						UpDir = DefaultSimInputs->UpDir;
-					}
+						QUICK_SCOPE_CYCLE_COUNTER(MoverStateMachine_OnSimulationTick_GenerateMoveFromMode);
+						CurrentMode->GenerateMove(WorkingSubstepStartData, SubTimeStep, OUT CombinedMove);
 
-					MovementMixer->MixProposedMoves(CombinedLayeredMove, UpDir, CombinedMove);
-				}
-				else if (bHasLayeredMoveContributions && !bHasModeMoveContribution)
-				{
-					CombinedMove = CombinedLayeredMove;
-				}
+						bHasModeMoveContribution = true;
+					}
 
-				// Apply any layered move finish velocity settings
-				if (CurrentLayeredMoves.bApplyResidualVelocity)
-				{
-					CombinedMove.LinearVelocity = CurrentLayeredMoves.ResidualVelocity;
-				}
-				if (CurrentLayeredMoves.ResidualClamping >= 0.0f)
-				{
-					CombinedMove.LinearVelocity = CombinedMove.LinearVelocity.GetClampedToMaxSize(CurrentLayeredMoves.ResidualClamping);
-				}
-				CurrentLayeredMoves.ResetResidualVelocity();
+					if (bHasModeMoveContribution && bHasLayeredMoveContributions)
+					{
+						FVector UpDir = FVector::UpVector;
+						if (const FChaosMoverSimulationDefaultInputs* DefaultSimInputs = Simulation->GetLocalSimInput().FindDataByType<FChaosMoverSimulationDefaultInputs>())
+						{
+							UpDir = DefaultSimInputs->UpDir;
+						}
 
-				// We need to replace this with some async equivalent (calling back to FSimulation? an optional FinalMoveProcessor object, a bit like the optional MoveMixer?)
-				// SyncTickParams->MoverComponent->ProcessGeneratedMovement.ExecuteIfBound(SubstepStartData, SubTimeStep, OUT CombinedMove);
+						MovementMixer->MixProposedMoves(CombinedLayeredMove, UpDir, CombinedMove);
+					}
+					else if (bHasLayeredMoveContributions && !bHasModeMoveContribution)
+					{
+						CombinedMove = CombinedLayeredMove;
+					}
 
-				// Execute the combined proposed move
-				{
-					// WorkingSimTickParams.MovingComps is left empty in the async case, so we don't access resources used by the concurrent gameplay thread
-					WorkingSimTickParams.StartState = WorkingSubstepStartData;
-					WorkingSimTickParams.SimBlackboard = SimBlackboard;
-					WorkingSimTickParams.TimeStep = SubTimeStep;
-					WorkingSimTickParams.ProposedMove = CombinedMove;
+					// Apply any layered move finish velocity settings
+					if (CurrentLayeredMoves.bApplyResidualVelocity)
+					{
+						CombinedMove.LinearVelocity = CurrentLayeredMoves.ResidualVelocity;
+					}
+					if (CurrentLayeredMoves.ResidualClamping >= 0.0f)
+					{
+						CombinedMove.LinearVelocity = CombinedMove.LinearVelocity.GetClampedToMaxSize(CurrentLayeredMoves.ResidualClamping);
+					}
+					CurrentLayeredMoves.ResetResidualVelocity();
 
-					// Check for any transitions, first those registered with the current movement mode, then global ones that could occur from any mode
-					FTransitionEvalResult EvalResult = FTransitionEvalResult::NoTransition;
-					TStrongObjectPtr<UBaseMovementModeTransition> TransitionToTrigger;
+					// We need to replace this with some async equivalent (calling back to FSimulation? an optional FinalMoveProcessor object, a bit like the optional MoveMixer?)
+					// SyncTickParams->MoverComponent->ProcessGeneratedMovement.ExecuteIfBound(SubstepStartData, SubTimeStep, OUT CombinedMove);
 
-					for (UBaseMovementModeTransition* Transition : CurrentMode->Transitions)
+					// Execute the combined proposed move
 					{
-						if (IsValid(Transition) && ((SubStepCount == 0) || !Transition->bFirstSubStepOnly))
-						{
-							EvalResult = Transition->Evaluate(WorkingSimTickParams);
+						// WorkingSimTickParams.MovingComps is left empty in the async case, so we don't access resources used by the concurrent gameplay thread
+						WorkingSimTickParams.StartState = WorkingSubstepStartData;
+						WorkingSimTickParams.SimBlackboard = SimBlackboard;
+						WorkingSimTickParams.TimeStep = SubTimeStep;
+						WorkingSimTickParams.ProposedMove = CombinedMove;
 
-							if (!EvalResult.NextMode.IsNone())
-							{
-								if (EvalResult.NextMode != CurrentModeName || Transition->bAllowModeReentry)
-								{
-									TransitionToTrigger = TStrongObjectPtr<UBaseMovementModeTransition>(Transition);
-									break;
-								}
-							}
-						}
-					}
+						// Check for any transitions, first those registered with the current movement mode, then global ones that could occur from any mode
+						FTransitionEvalResult EvalResult = FTransitionEvalResult::NoTransition;
+						TStrongObjectPtr<UBaseMovementModeTransition> TransitionToTrigger;
 
-					if (TransitionToTrigger == nullptr)
-					{
-						for (TWeakObjectPtr<UBaseMovementModeTransition> TransitionWeakPtr : GlobalTransitions)
+						for (UBaseMovementModeTransition* Transition : CurrentMode->Transitions)
 						{
-							TStrongObjectPtr<UBaseMovementModeTransition> Transition = TransitionWeakPtr.Pin();
-							if (Transition && ((SubStepCount == 0) || !Transition->bFirstSubStepOnly))
+							if (IsValid(Transition) && ((SubStepCount == 0) || !Transition->bFirstSubStepOnly))
 							{
 								EvalResult = Transition->Evaluate(WorkingSimTickParams);
 
@@ -539,51 +521,72 @@ namespace UE::ChaosMover
 								{
 									if (EvalResult.NextMode != CurrentModeName || Transition->bAllowModeReentry)
 									{
-										TransitionToTrigger = Transition;
+										TransitionToTrigger = TStrongObjectPtr<UBaseMovementModeTransition>(Transition);
 										break;
 									}
 								}
 							}
 						}
+
+						if (TransitionToTrigger == nullptr)
+						{
+							for (TWeakObjectPtr<UBaseMovementModeTransition> TransitionWeakPtr : GlobalTransitions)
+							{
+								TStrongObjectPtr<UBaseMovementModeTransition> Transition = TransitionWeakPtr.Pin();
+								if (Transition && ((SubStepCount == 0) || !Transition->bFirstSubStepOnly))
+								{
+									EvalResult = Transition->Evaluate(WorkingSimTickParams);
+
+									if (!EvalResult.NextMode.IsNone())
+									{
+										if (EvalResult.NextMode != CurrentModeName || Transition->bAllowModeReentry)
+										{
+											TransitionToTrigger = Transition;
+											break;
+										}
+									}
+								}
+							}
+						}
+
+						if (TransitionToTrigger && !EvalResult.NextMode.IsNone())
+						{
+							OutputState.MovementEndState.NextModeName = EvalResult.NextMode;
+							OutputState.MovementEndState.RemainingMs = WorkingSimTickParams.TimeStep.StepMs; 	// Pass all remaining time to next mode
+							TransitionToTrigger->Trigger(WorkingSimTickParams);
+						}
+						else
+						{
+							QUICK_SCOPE_CYCLE_COUNTER(MoverStateMachine_OnSimulationTick_CurrentModeSimulationTick);
+							CurrentMode->SimulationTick(WorkingSimTickParams, OutputState);
+						}
+
+						OutputState.MovementEndState.RemainingMs = FMath::Clamp(OutputState.MovementEndState.RemainingMs, 0.0f, SubTimeStep.StepMs);
 					}
 
-					if (TransitionToTrigger && !EvalResult.NextMode.IsNone())
+					QueueNextMode(OutputState.MovementEndState.NextModeName);
+
+					// Check if all of the time for this Substep was refunded
+					if (FMath::IsNearlyEqual(SubTimeStep.StepMs, OutputState.MovementEndState.RemainingMs, UE_KINDA_SMALL_NUMBER))
 					{
-						OutputState.MovementEndState.NextModeName = EvalResult.NextMode;
-						OutputState.MovementEndState.RemainingMs = WorkingSimTickParams.TimeStep.StepMs; 	// Pass all remaining time to next mode
-						TransitionToTrigger->Trigger(WorkingSimTickParams);
+						NumConsecutiveFullRefundedSubsteps++;
+						// if we've done this sub step a lot before go ahead and just advance time to avoid freezing editor
+						if (NumConsecutiveFullRefundedSubsteps >= MaxConsecutiveFullRefundedSubsteps)
+						{
+							UE_LOG(LogChaosMover, Warning, TEXT("Movement mode %s and %s on %s are stuck giving time back to each other. Overriding to advance to next substep."),
+								*CurrentModeName.ToString(),
+								*OutputState.MovementEndState.NextModeName.ToString(),
+								*OwnerActorName);
+							TotalUsedMs += SubTimeStep.StepMs;
+						}
 					}
 					else
 					{
-						QUICK_SCOPE_CYCLE_COUNTER(MoverStateMachine_OnSimulationTick_CurrentModeSimulationTick);
-						CurrentMode->SimulationTick(WorkingSimTickParams, OutputState);
+						NumConsecutiveFullRefundedSubsteps = 0;
 					}
 
-					OutputState.MovementEndState.RemainingMs = FMath::Clamp(OutputState.MovementEndState.RemainingMs, 0.0f, SubTimeStep.StepMs);
+					//GEngine->AddOnScreenDebugMessage(-1, -0.1f, FColor::White, FString::Printf(TEXT("NextModeName: %s  Queued: %s"), *Output.MovementEndState.NextModeName.ToString(), *NextModeName.ToString()));
 				}
-
-				QueueNextMode(OutputState.MovementEndState.NextModeName);
-
-				// Check if all of the time for this Substep was refunded
-				if (FMath::IsNearlyEqual(SubTimeStep.StepMs, OutputState.MovementEndState.RemainingMs, UE_KINDA_SMALL_NUMBER))
-				{
-					NumConsecutiveFullRefundedSubsteps++;
-					// if we've done this sub step a lot before go ahead and just advance time to avoid freezing editor
-					if (NumConsecutiveFullRefundedSubsteps >= MaxConsecutiveFullRefundedSubsteps)
-					{
-						UE_LOG(LogChaosMover, Warning, TEXT("Movement mode %s and %s on %s are stuck giving time back to each other. Overriding to advance to next substep."),
-							*CurrentModeName.ToString(),
-							*OutputState.MovementEndState.NextModeName.ToString(),
-							*OwnerActorName);
-						TotalUsedMs += SubTimeStep.StepMs;
-					}
-				}
-				else
-				{
-					NumConsecutiveFullRefundedSubsteps = 0;
-				}
-
-				//GEngine->AddOnScreenDebugMessage(-1, -0.1f, FColor::White, FString::Printf(TEXT("NextModeName: %s  Queued: %s"), *Output.MovementEndState.NextModeName.ToString(), *NextModeName.ToString()));
 			}
 
 			// Switch modes if necessary (note that this will allow exit/enter on the same state)
