[ChaosMover] Crouch Flow Improvements

1. Removed CrouchModifierHandle from transition (modifier handles are not NetSerialized and triggered spurious modifier cancellations)

2a. Toggle bWantsToCrouch based on current bWantsToCrouch and not previous state (crouch toggling will now track the input presses correctly)
2b. Networked crouch input is reverted to being edge based (not sent every frame like before)
2c. Removed the Crouch/Uncrouch calls in Modifier Events with silently updating the crouch intent (prevents input state feedback loop)

3. Update CMC's crouch state every frame (to animate correctly)

diff --git a/Engine/Plugins/Experimental/ChaosMover/Source/ChaosMover/Private/Character/ChaosCharacterMoverComponent.cpp b/Engine/Plugins/Experimental/ChaosMover/Source/ChaosMover/Private/Character/ChaosCharacterMoverComponent.cpp
--- a/Engine/Plugins/Experimental/ChaosMover/Source/ChaosMover/Private/Character/ChaosCharacterMoverComponent.cpp
+++ b/Engine/Plugins/Experimental/ChaosMover/Source/ChaosMover/Private/Character/ChaosCharacterMoverComponent.cpp
@@ -121,11 +121,16 @@ void UChaosCharacterMoverComponent::ProduceInput(const int32 DeltaTimeMS, FMover
 		ClearQueuedMode();
 	}
 
-	// TODO @Harsha Making bWantsToCrouch sticky because the server is missing the one time crouch input.
-	// Investigate the input drops and after the fix, revert this change back to making it non-sticky.
-	bWantsToCrouch &= CanCrouch();
-	FChaosMoverCrouchInputs& CrouchInputs = Cmd->InputCollection.FindOrAddMutableDataByType<FChaosMoverCrouchInputs>();
-	CrouchInputs.bWantsToCrouch = bWantsToCrouch;
+	if (bCrouchInputPending)
+	{
+		if ((bWantsToCrouch && CanCrouch()) || !bWantsToCrouch)
+		{
+			FChaosMoverCrouchInputs& CrouchInputs = Cmd->InputCollection.FindOrAddMutableDataByType<FChaosMoverCrouchInputs>();
+			CrouchInputs.bWantsToCrouch = bWantsToCrouch;
+		}
+
+		bCrouchInputPending = false;
+	}
 
 	if (bCancelMovementOverrides)
 	{
@@ -166,6 +171,16 @@ void UChaosCharacterMoverComponent::UnCrouch()
 
 void UChaosCharacterMoverComponent::OnStanceModified(const FStanceModifiedEventData& EventData)
 {
+	// Update bWantsToCrouch without triggering networked crouch input
+	if (EventData.NewStance == EStanceMode::Invalid)
+	{
+		bWantsToCrouch = false;
+	}
+	else if (EventData.NewStance == EStanceMode::Crouch)
+	{
+		bWantsToCrouch = true;
+	}
+	
 	OnStanceChanged.Broadcast(EventData.OldStance, EventData.NewStance);
 }
 
diff --git a/Engine/Plugins/Experimental/ChaosMover/Source/ChaosMover/Private/Character/Modifiers/ChaosStanceModifier.cpp b/Engine/Plugins/Experimental/ChaosMover/Source/ChaosMover/Private/Character/Modifiers/ChaosStanceModifier.cpp
--- a/Engine/Plugins/Experimental/ChaosMover/Source/ChaosMover/Private/Character/Modifiers/ChaosStanceModifier.cpp
+++ b/Engine/Plugins/Experimental/ChaosMover/Source/ChaosMover/Private/Character/Modifiers/ChaosStanceModifier.cpp
@@ -64,8 +64,13 @@ void FChaosStanceModifier::UpdateStance(const FMovementModifierParams_Async& Asy
 {
 	if (NewStance == ActiveStance)
 	{
-		UE_LOG(LogChaosMover, Verbose, TEXT("Skipping stance change - not required"));
-		return;
+		// When resimulating, NewStance == ActiveStance is inherently always true because the server auth sync state will have the correct active state
+		// We will need to continue with Updating the capsule irrespective of this check
+		if (!AsyncParams.TimeStep->bIsResimulating)
+		{
+			UE_LOG(LogChaosMover, Verbose, TEXT("Skipping stance change - not required"));
+			return;
+		}
 	}
 
 	if (!AsyncParams.IsValid())
@@ -183,19 +188,6 @@ void FChaosStanceModifier::OnStart_Async(const FMovementModifierParams_Async& As
 {
 	UpdateStance(AsyncParams, EStanceMode::Crouch);
 	
-	// Propagate crouch state to GT if it's triggered on the PT (rollback or resim)
-	// TODO @Harsha this is not necessary if bWantsToCrouch is non-sticky, stale bWantsToCrouch will not matter then
-	if (UChaosMoverSimulation* Sim = Cast<UChaosMoverSimulation>(AsyncParams.Simulation))
-	{
-		auto EnableWantsToCrouch = [] (const FMoverSimulationEventData& Data, const FMoverSimEventGameThreadContext& GTContext) {
-			if (UCharacterMoverComponent* CharMoverComp = Cast<UCharacterMoverComponent>(GTContext.MoverComp))
-			{
-				CharMoverComp->Crouch();
-			}
-		};
-		Sim->AddEvent(MakeShared<FStanceModifierModeChangedEventData>(*AsyncParams.TimeStep, CurrentModeName, AsyncParams.SyncState->MovementMode, EnableWantsToCrouch));
-	}
-
 	// Set the current mode - this is used to check if the mode changes
 	CurrentModeName = AsyncParams.SyncState->MovementMode;
 }
@@ -204,19 +196,6 @@ void FChaosStanceModifier::OnEnd_Async(const FMovementModifierParams_Async& Asyn
 {
 	UpdateStance(AsyncParams, EStanceMode::Invalid);
 	
-	// Propagate uncrouch state to GT if it's triggered on the PT (rollback, resim or uncrouching on mode change)
-	// TODO @Harsha this is not necessary if bWantsToCrouch is non-sticky, stale bWantsToCrouch will not matter then
-	if (UChaosMoverSimulation* Sim = Cast<UChaosMoverSimulation>(AsyncParams.Simulation))
-	{
-		auto DisableWantsToCrouch = [](const FMoverSimulationEventData& Data, const FMoverSimEventGameThreadContext& GTContext) {
-			if (UCharacterMoverComponent* CharMoverComp = Cast<UCharacterMoverComponent>(GTContext.MoverComp))
-			{
-				CharMoverComp->UnCrouch();
-			}
-		};
-		Sim->AddEvent(MakeShared<FStanceModifierModeChangedEventData>(*AsyncParams.TimeStep, CurrentModeName, AsyncParams.SyncState->MovementMode, DisableWantsToCrouch));
-	}
-
 	CurrentModeName = NAME_None;
 }
 
diff --git a/Engine/Plugins/Experimental/ChaosMover/Source/ChaosMover/Private/Character/Transitions/ChaosCharacterCrouchCheck.cpp b/Engine/Plugins/Experimental/ChaosMover/Source/ChaosMover/Private/Character/Transitions/ChaosCharacterCrouchCheck.cpp
--- a/Engine/Plugins/Experimental/ChaosMover/Source/ChaosMover/Private/Character/Transitions/ChaosCharacterCrouchCheck.cpp
+++ b/Engine/Plugins/Experimental/ChaosMover/Source/ChaosMover/Private/Character/Transitions/ChaosCharacterCrouchCheck.cpp
@@ -89,7 +89,7 @@ void UChaosCharacterCrouchCheck::Crouch(const FSimulationTickParams& TickParams)
 	StanceModifier->MaxSpeedOverride = MaxSpeed;
 	StanceModifier->bCancelOnModeChange = bCancelOnModeChange;
 
-	CrouchModifierHandle = Simulation->QueueMovementModifier(MoveTemp(StanceModifier));
+	Simulation->QueueMovementModifier(MoveTemp(StanceModifier));
 }
 
 void UChaosCharacterCrouchCheck::Uncrouch(const FSimulationTickParams& TickParams)
@@ -98,15 +98,9 @@ void UChaosCharacterCrouchCheck::Uncrouch(const FSimulationTickParams& TickParam
 
 	check(Simulation);
 
-	if (CrouchModifierHandle.IsValid())
-	{
-		Simulation->CancelModifierFromHandle(CrouchModifierHandle);
-		CrouchModifierHandle.Invalidate();
-	}
-	else if (const FChaosStanceModifier* Modifier = Simulation->FindMovementModifierByType<FChaosStanceModifier>())
+	if (const FChaosStanceModifier* Modifier = Simulation->FindMovementModifierByType<FChaosStanceModifier>())
 	{
 		Simulation->CancelModifierFromHandle(Modifier->GetHandle());
-		CrouchModifierHandle.Invalidate();
 	}
 }
 
@@ -132,14 +126,6 @@ FTransitionEvalResult UChaosCharacterCrouchCheck::Evaluate_Implementation(const
 		if (Simulation->HasMovementMode(TickParams.StartState.SyncState.MovementMode))
 		{
 			bool bIsCrouching = Simulation->HasGameplayTag(Mover_IsCrouching, true);
-			if (bIsCrouching && !CrouchModifierHandle.IsValid())
-			{
-				if (const FChaosStanceModifier* Modifier = Simulation->FindMovementModifierByType<FChaosStanceModifier>())
-				{
-					Simulation->CancelModifierFromHandle(Modifier->GetHandle());
-					bIsCrouching = false;
-				}
-			}
 
 			if (bIsCrouching && !CrouchInputs->bWantsToCrouch && CanUncrouch(TickParams))
 			{
diff --git a/Engine/Plugins/Experimental/ChaosMover/Source/ChaosMover/Public/ChaosMover/ChaosMoverSimulationTypes.h b/Engine/Plugins/Experimental/ChaosMover/Source/ChaosMover/Public/ChaosMover/ChaosMoverSimulationTypes.h
--- a/Engine/Plugins/Experimental/ChaosMover/Source/ChaosMover/Public/ChaosMover/ChaosMoverSimulationTypes.h
+++ b/Engine/Plugins/Experimental/ChaosMover/Source/ChaosMover/Public/ChaosMover/ChaosMoverSimulationTypes.h
@@ -366,41 +366,6 @@ struct FStanceModifiedEventData : public FMoverSimulationEventData
 	EStanceMode NewStance = EStanceMode::Invalid;
 };
 
-USTRUCT()
-struct FStanceModifierModeChangedEventData : public FMoverSimulationEventData
-{
-	GENERATED_BODY()
-	
-	FStanceModifierModeChangedEventData(const FMoverTimeStep& InEventTime, FName InOldMode, FName InNewMode, FEventProcessedCallbackPtr InEventProcessedCallback = nullptr)
-	: FMoverSimulationEventData(InEventTime, InEventProcessedCallback)
-	, OldMode(InOldMode)
-	, NewMode(InNewMode)
-	{
-	}
-	FStanceModifierModeChangedEventData() {}
-
-	virtual UScriptStruct* GetScriptStruct() const override
-	{
-		return FStanceModifierModeChangedEventData::StaticStruct();
-	}
-
-	virtual bool IsEqual(const FMoverSimulationEventData& Other, int32 ServerFrameDiffTolerance = 0) const override
-	{
-		if (Super::IsEqual(Other, ServerFrameDiffTolerance))
-		{
-			const FStanceModifierModeChangedEventData& TypedOther = static_cast<const FStanceModifierModeChangedEventData&>(Other);
-			return (OldMode == TypedOther.OldMode) && (NewMode == TypedOther.NewMode);
-		}
-		else
-		{
-			return false;
-		}
-	}
-
-	FName OldMode = NAME_None;
-	FName NewMode = NAME_None;
-};
-
 // Version of a FScheduledInstantMovementEffect with Issuance Server Frame
 USTRUCT()
 struct FChaosScheduledInstantMovementEffect
diff --git a/Engine/Plugins/Experimental/ChaosMover/Source/ChaosMover/Public/ChaosMover/Character/Transitions/ChaosCharacterCrouchCheck.h b/Engine/Plugins/Experimental/ChaosMover/Source/ChaosMover/Public/ChaosMover/Character/Transitions/ChaosCharacterCrouchCheck.h
--- a/Engine/Plugins/Experimental/ChaosMover/Source/ChaosMover/Public/ChaosMover/Character/Transitions/ChaosCharacterCrouchCheck.h
+++ b/Engine/Plugins/Experimental/ChaosMover/Source/ChaosMover/Public/ChaosMover/Character/Transitions/ChaosCharacterCrouchCheck.h
@@ -3,7 +3,6 @@
 #pragma once
 
 #include "ChaosMover/ChaosMovementModeTransition.h"
-#include "MovementModifier.h"
 
 #include "ChaosCharacterCrouchCheck.generated.h"
 
@@ -61,8 +60,6 @@ protected:
 	CHAOSMOVER_API virtual void Crouch(const FSimulationTickParams& Params);
 	CHAOSMOVER_API virtual void Uncrouch(const FSimulationTickParams& Params);
 
-	FMovementModifierHandle CrouchModifierHandle;
-
 	float OriginalCapsuleHalfHeight = 0.0f;
 	float OriginalCapsuleRadius = 0.0f;
 
diff --git a/Engine/Plugins/Experimental/Mover/Source/Mover/Public/DefaultMovementSet/CharacterMoverComponent.h b/Engine/Plugins/Experimental/Mover/Source/Mover/Public/DefaultMovementSet/CharacterMoverComponent.h
--- a/Engine/Plugins/Experimental/Mover/Source/Mover/Public/DefaultMovementSet/CharacterMoverComponent.h
+++ b/Engine/Plugins/Experimental/Mover/Source/Mover/Public/DefaultMovementSet/CharacterMoverComponent.h
@@ -98,6 +98,10 @@ public:
 	UFUNCTION(BlueprintCallable, Category = Mover)
 	UE_API virtual void UnCrouch();
 
+	// Get Crouch Intent
+	UFUNCTION(BlueprintCallable, Category = Mover)
+	bool GetCrouchIntent() const { return bWantsToCrouch; };
+
 	// Broadcast when this actor changes stances.
 	UPROPERTY(BlueprintAssignable, Category = Mover)
 	FMover_OnStanceChanged OnStanceChanged;
