Fix for resim causing issues for characters on turntables

diff --git a/Engine/Plugins/Experimental/ChaosMover/Source/ChaosMover/Private/ChaosMoverSimulation.cpp b/Engine/Plugins/Experimental/ChaosMover/Source/ChaosMover/Private/ChaosMoverSimulation.cpp
--- a/Engine/Plugins/Experimental/ChaosMover/Source/ChaosMover/Private/ChaosMoverSimulation.cpp
+++ b/Engine/Plugins/Experimental/ChaosMover/Source/ChaosMover/Private/ChaosMoverSimulation.cpp
@@ -379,9 +379,12 @@ void UChaosMoverSimulation::PreSimulationTickCharacter(const IChaosCharacterMove
 		}
 	}
 
-	if (UE::ChaosMover::CVars::bEnablePreSimGroundCheck)
+	// Update the floor before we run any update.
+	// bEnablePreSimGroundCheck OFF allows us to save performance by relying on last frame's result until the floor gets updated
+	// However when we rollback we invalidate the results in the blackboard since they are for the frame we rewound from,
+	// which leaves us without any floor info, causing bugs. So we need to recalculate them in the first resim frame case.
+	if (UE::ChaosMover::CVars::bEnablePreSimGroundCheck || (TimeStep.bIsResimulating && TimeStep.bIsFirstResimFrame))
 	{
-		// Update the floor before we run any update.
 		CharacterMode.UpdateCurrentFloor(TimeStep);
 	}
 }
