[Chaos Mover] Add optional server inputs for server side launch and other effects

diff --git a/Engine/Plugins/Experimental/ChaosMover/Source/ChaosMover/Private/Backends/ChaosMoverBackend.cpp b/Engine/Plugins/Experimental/ChaosMover/Source/ChaosMover/Private/Backends/ChaosMoverBackend.cpp
--- a/Engine/Plugins/Experimental/ChaosMover/Source/ChaosMover/Private/Backends/ChaosMoverBackend.cpp
+++ b/Engine/Plugins/Experimental/ChaosMover/Source/ChaosMover/Private/Backends/ChaosMoverBackend.cpp
@@ -3,8 +3,10 @@
 #include "ChaosMover/Backends/ChaosMoverBackend.h"
 
 #include "Backends/ChaosMoverSubsystem.h"
+#include "ChaosMover/ChaosMoverConsoleVariables.h"
 #include "ChaosMover/ChaosMoverLog.h"
 #include "ChaosMover/ChaosMoverSimulation.h"
+#include "ChaosMover/Character/ChaosCharacterMoverComponent.h"
 #include "Components/PrimitiveComponent.h"
 #include "Framework/Threading.h"
 #include "GameFramework/Pawn.h"
@@ -484,12 +486,21 @@ void UChaosMoverBackendComponent::ProduceInputData(int32 PhysicsStep, int32 NumS
 		return;
 	}
 
-	bool bIsGeneratingInputsLocally = !NetworkPhysicsComponent || NetworkPhysicsComponent->IsLocallyControlled();
+	const bool bIsGeneratingInputsLocally = !NetworkPhysicsComponent || NetworkPhysicsComponent->IsLocallyControlled();
 	if (bIsGeneratingInputsLocally)
 	{
 		GenerateInput(TimeStep, InputData);
 	}
 
+	if (UE::ChaosMover::CVars::bEnableServerLaunchOverride)
+	{
+		const bool bIsServer = NetworkPhysicsComponent && NetworkPhysicsComponent->HasServerWorld();
+		if (bIsServer)
+		{
+			GenerateServerInput(TimeStep, InputData);
+		}
+	}
+
 	UMoverComponent& MoverComp = GetMoverComponent();
 	// Add all queued instant effects to the input so it can be networked
 	const TArray<FScheduledInstantMovementEffect>& ScheduledInstantMovementEffects = MoverComp.GetQueuedInstantMovementEffects();
@@ -560,6 +571,16 @@ void UChaosMoverBackendComponent::GenerateInput(const FMoverTimeStep& TimeStep,
 	MoverComp.ProduceInput(TimeStep.StepMs, &InputData.InputCmd);
 }
 
+void UChaosMoverBackendComponent::GenerateServerInput(const FMoverTimeStep& TimeStep, UE::ChaosMover::FSimulationInputData& InputData)
+{
+	UMoverComponent& MoverComp = GetMoverComponent();
+
+	if (UChaosCharacterMoverComponent* CharacterMoverComp = Cast<UChaosCharacterMoverComponent>(&MoverComp))
+	{
+		CharacterMoverComp->ProduceServerInput(TimeStep.StepMs, &InputData.InputCmd);
+	}
+}
+
 void UChaosMoverBackendComponent::InjectInstantMovementEffectsIntoInput(const FMoverTimeStep& TimeStep, UE::ChaosMover::FSimulationInputData& InputData)
 {
 	UMoverComponent& MoverComp = GetMoverComponent();
diff --git a/Engine/Plugins/Experimental/ChaosMover/Source/ChaosMover/Private/ChaosMoverConsoleVariables.cpp b/Engine/Plugins/Experimental/ChaosMover/Source/ChaosMover/Private/ChaosMoverConsoleVariables.cpp
--- a/Engine/Plugins/Experimental/ChaosMover/Source/ChaosMover/Private/ChaosMoverConsoleVariables.cpp
+++ b/Engine/Plugins/Experimental/ChaosMover/Source/ChaosMover/Private/ChaosMoverConsoleVariables.cpp
@@ -55,6 +55,10 @@ namespace UE::ChaosMover
 		FAutoConsoleVariableRef CVarChaosMoverEnablePreSimGroundCheck(TEXT("ChaosMover.EnablePreSimGroundCheck"),
 			bEnablePreSimGroundCheck, TEXT(""));
 
+		bool bEnableServerLaunchOverride = false;
+		FAutoConsoleVariableRef CVarChaosMoverEnableServerLaunchOverride(TEXT("ChaosMover.EnableServerLaunchOverride"),
+			bEnableServerLaunchOverride, TEXT(""));
+
 		static FAutoConsoleCommandWithWorld TeleportOnlyClientCommand
 		(
 			TEXT("ChaosMover.Debug.TeleportClientOnly"),
diff --git a/Engine/Plugins/Experimental/ChaosMover/Source/ChaosMover/Private/ChaosMoverSimulation.cpp b/Engine/Plugins/Experimental/ChaosMover/Source/ChaosMover/Private/ChaosMoverSimulation.cpp
--- a/Engine/Plugins/Experimental/ChaosMover/Source/ChaosMover/Private/ChaosMoverSimulation.cpp
+++ b/Engine/Plugins/Experimental/ChaosMover/Source/ChaosMover/Private/ChaosMoverSimulation.cpp
@@ -210,6 +210,30 @@ void UChaosMoverSimulation::ProcessInputs(int32 PhysicsStep, const FMoverTimeSte
 	// On the client the input data is read out to be sent to the server so update the data that is read.
 	if (bInputCmdOverridden)
 	{
+		// If we are on the server we can also have inputs from actions applied to the mover component
+		// e.g. Launch. If so, use the server versions
+		if (UE::ChaosMover::CVars::bEnableServerLaunchOverride)
+		{
+			if (const FChaosMoverSimulationDefaultInputs* SimInputs = LocalSimInput.FindDataByType<FChaosMoverSimulationDefaultInputs>())
+			{
+				if (SimInputs->AsyncNetworkPhysicsComponent && SimInputs->AsyncNetworkPhysicsComponent->IsServer())
+				{
+					if (const FChaosMoverLaunchInputs* NetLaunchInputs = NetInputCmd.InputCollection.FindDataByType<FChaosMoverLaunchInputs>())
+					{
+						UE_LOG(LogTemp, Warning, TEXT("Removed Launch. Frame=%u, Value=%s"), TimeStep.ServerFrame, *NetLaunchInputs->LaunchVelocityOrImpulse.ToCompactString());
+						NetInputCmd.InputCollection.RemoveDataByType(FChaosMoverLaunchInputs::StaticStruct());
+					}
+
+					if (const FChaosMoverLaunchInputs* LocalLaunchInputs = InputData.InputCmd.InputCollection.FindDataByType<FChaosMoverLaunchInputs>())
+					{
+						UE_LOG(LogTemp, Warning, TEXT("Added Launch. Frame=%u, Value=%s"), TimeStep.ServerFrame, *LocalLaunchInputs->LaunchVelocityOrImpulse.ToCompactString());
+						FChaosMoverLaunchInputs& LaunchInputs = NetInputCmd.InputCollection.FindOrAddMutableDataByType<FChaosMoverLaunchInputs>();
+						LaunchInputs = *LocalLaunchInputs;
+					}
+				}
+			}
+		}
+
 		InputData.InputCmd = NetInputCmd;
 		bInputCmdOverridden = false;
 	}
diff --git a/Engine/Plugins/Experimental/ChaosMover/Source/ChaosMover/Private/Character/ChaosCharacterMoverComponent.cpp b/Engine/Plugins/Experimental/ChaosMover/Source/ChaosMover/Private/Character/ChaosCharacterMoverComponent.cpp
--- a/Engine/Plugins/Experimental/ChaosMover/Source/ChaosMover/Private/Character/ChaosCharacterMoverComponent.cpp
+++ b/Engine/Plugins/Experimental/ChaosMover/Source/ChaosMover/Private/Character/ChaosCharacterMoverComponent.cpp
@@ -99,9 +99,22 @@ bool UChaosCharacterMoverComponent::TryGetFloorCheckHitResult(FHitResult& OutHit
 
 void UChaosCharacterMoverComponent::ProduceInput(const int32 DeltaTimeMS, FMoverInputCmdContext* Cmd)
 {
-	Super::ProduceInput(DeltaTimeMS, Cmd);
+	ProduceInputImpl(DeltaTimeMS, Cmd);
+}
+
+void UChaosCharacterMoverComponent::ProduceServerInput(const int32 DeltaTimeMS, FMoverInputCmdContext* Cmd)
+{
+	ProduceInputImpl(DeltaTimeMS, Cmd);
+}
 
-	check(Cmd);
+void UChaosCharacterMoverComponent::ProduceInputImpl(const int32 DeltaTimeMS, FMoverInputCmdContext* Cmd)
+{
+	if (!ensure(Cmd))
+	{
+		return;
+	}
+
+	Super::ProduceInput(DeltaTimeMS, Cmd);
 
 	// Add launch input data if launch velocity is set
 	if (!LaunchVelocityOrImpulse.IsZero())
diff --git a/Engine/Plugins/Experimental/ChaosMover/Source/ChaosMover/Public/ChaosMover/Backends/ChaosMoverBackend.h b/Engine/Plugins/Experimental/ChaosMover/Source/ChaosMover/Public/ChaosMover/Backends/ChaosMoverBackend.h
--- a/Engine/Plugins/Experimental/ChaosMover/Source/ChaosMover/Public/ChaosMover/Backends/ChaosMoverBackend.h
+++ b/Engine/Plugins/Experimental/ChaosMover/Source/ChaosMover/Public/ChaosMover/Backends/ChaosMoverBackend.h
@@ -64,6 +64,7 @@ protected:
 	CHAOSMOVER_API void DestroyActuationConstraint();
 
 	CHAOSMOVER_API virtual void GenerateInput(const FMoverTimeStep& TimeStep, UE::ChaosMover::FSimulationInputData& InputData);
+	CHAOSMOVER_API virtual void GenerateServerInput(const FMoverTimeStep& TimeStep, UE::ChaosMover::FSimulationInputData& InputData);
 
 	CHAOSMOVER_API UMoverComponent& GetMoverComponent() const;
 	CHAOSMOVER_API Chaos::FPhysicsSolver* GetPhysicsSolver() const;
diff --git a/Engine/Plugins/Experimental/ChaosMover/Source/ChaosMover/Public/ChaosMover/ChaosMoverConsoleVariables.h b/Engine/Plugins/Experimental/ChaosMover/Source/ChaosMover/Public/ChaosMover/ChaosMoverConsoleVariables.h
--- a/Engine/Plugins/Experimental/ChaosMover/Source/ChaosMover/Public/ChaosMover/ChaosMoverConsoleVariables.h
+++ b/Engine/Plugins/Experimental/ChaosMover/Source/ChaosMover/Public/ChaosMover/ChaosMoverConsoleVariables.h
@@ -20,5 +20,6 @@ namespace UE::ChaosMover
 
 		// Temporary
 		CHAOSMOVER_API extern bool bEnablePreSimGroundCheck;
+		CHAOSMOVER_API extern bool bEnableServerLaunchOverride;
 	}
 };
diff --git a/Engine/Plugins/Experimental/ChaosMover/Source/ChaosMover/Public/ChaosMover/Character/ChaosCharacterMoverComponent.h b/Engine/Plugins/Experimental/ChaosMover/Source/ChaosMover/Public/ChaosMover/Character/ChaosCharacterMoverComponent.h
--- a/Engine/Plugins/Experimental/ChaosMover/Source/ChaosMover/Public/ChaosMover/Character/ChaosCharacterMoverComponent.h
+++ b/Engine/Plugins/Experimental/ChaosMover/Source/ChaosMover/Public/ChaosMover/Character/ChaosCharacterMoverComponent.h
@@ -26,6 +26,7 @@ public:
 
 	CHAOSMOVER_API virtual bool TryGetFloorCheckHitResult(FHitResult& OutHitResult) const override;
 	CHAOSMOVER_API virtual void ProduceInput(const int32 DeltaTimeMS, FMoverInputCmdContext* Cmd) override;
+	CHAOSMOVER_API virtual void ProduceServerInput(const int32 DeltaTimeMS, FMoverInputCmdContext* Cmd);
 	CHAOSMOVER_API virtual void DoQueueNextMode(FName DesiredModeName, bool bShouldReenter=false) override;
 
 	CHAOSMOVER_API virtual void Crouch() override;
@@ -65,6 +66,8 @@ protected:
 	CHAOSMOVER_API void NotifyOnLanded(const FLandedEventData& LandedEvent);
 	CHAOSMOVER_API virtual void OnLanded(const FLandedEventData& LandedEvent);
 
+	CHAOSMOVER_API virtual void ProduceInputImpl(const int32 DeltaTimeMS, FMoverInputCmdContext* Cmd);
+
 	FName ModeToOverrideSettings = NAME_None;
 	float MaxSpeedOverride = 0.0f;
 	float AccelerationOverride = 0.0f;
