[ChaosMover] Fix for a crash when accessing SharedSettings, by making it a weak pointer and pinning a strong pointer upon usage.

diff --git a/Engine/Plugins/Experimental/ChaosMover/Source/ChaosMover/Private/Character/Modes/ChaosCharacterMovementMode.cpp b/Engine/Plugins/Experimental/ChaosMover/Source/ChaosMover/Private/Character/Modes/ChaosCharacterMovementMode.cpp
--- a/Engine/Plugins/Experimental/ChaosMover/Source/ChaosMover/Private/Character/Modes/ChaosCharacterMovementMode.cpp
+++ b/Engine/Plugins/Experimental/ChaosMover/Source/ChaosMover/Private/Character/Modes/ChaosCharacterMovementMode.cpp
@@ -43,15 +43,11 @@ void UChaosCharacterMovementMode::OnRegistered(const FName ModeName)
 		QueryRadius = QueryRadiusOverride.GetValue();
 	}
 
-	SharedSettings = GetMoverComponent()->FindSharedSettings<USharedChaosCharacterMovementSettings>();
-	ensureMsgf(SharedSettings, TEXT("Failed to find instance of USharedChaosCharacterMovementSettings on %s. Movement may not function properly."), *GetPathNameSafe(this));
-}
-
-void UChaosCharacterMovementMode::OnUnregistered()
-{
-	SharedSettings = nullptr;
-
-	Super::OnUnregistered();
+	TObjectPtr<const USharedChaosCharacterMovementSettings> SharedSettingsPtr = GetMoverComponent()->FindSharedSettings<USharedChaosCharacterMovementSettings>();
+	if (ensureMsgf(SharedSettingsPtr, TEXT("Failed to find instance of USharedChaosCharacterMovementSettings on %s. Movement may not function properly."), *GetPathNameSafe(this)))
+	{
+		SharedSettingsWeakPtr = TWeakObjectPtr<const USharedChaosCharacterMovementSettings>(SharedSettingsPtr);
+	}
 }
 
 void UChaosCharacterMovementMode::SetTargetHeightOverride(float InTargetHeight)
@@ -106,9 +102,9 @@ void UChaosCharacterMovementMode::UpdateConstraintSettings(Chaos::FCharacterGrou
 
 float UChaosCharacterMovementMode::GetMaxWalkSlopeCosine() const
 {
-	if (const USharedChaosCharacterMovementSettings* SharedSettingsPtr = SharedSettings.Get())
+	if (TStrongObjectPtr<const USharedChaosCharacterMovementSettings> SharedSettingsStrongPtr = SharedSettingsWeakPtr.Pin())
 	{
-		return SharedSettingsPtr->GetMaxWalkableSlopeCosine();
+		return SharedSettingsStrongPtr->GetMaxWalkableSlopeCosine();
 	}
 
 	return 0.707f;
@@ -120,9 +116,9 @@ float UChaosCharacterMovementMode::GetMaxSpeed() const
 	{
 		return MaxSpeedOverride.GetValue();
 	}
-	else if (const USharedChaosCharacterMovementSettings* SharedSettingsPtr = SharedSettings.Get())
+	else if (TStrongObjectPtr<const USharedChaosCharacterMovementSettings> SharedSettingsStrongPtr = SharedSettingsWeakPtr.Pin())
 	{
-		return SharedSettingsPtr->MaxSpeed;
+		return SharedSettingsStrongPtr->MaxSpeed;
 	}
 
 	UE_LOG(LogChaosMover, Warning, TEXT("Invalid max speed on ChaosCharacterMoverComponent"));
@@ -145,9 +141,9 @@ float UChaosCharacterMovementMode::GetAcceleration() const
 	{
 		return AccelerationOverride.GetValue();
 	}
-	else if (const USharedChaosCharacterMovementSettings* SharedSettingsPtr = SharedSettings.Get())
+	else if (TStrongObjectPtr<const USharedChaosCharacterMovementSettings> SharedSettingsStrongPtr = SharedSettingsWeakPtr.Pin())
 	{
-		return SharedSettingsPtr->Acceleration;
+		return SharedSettingsStrongPtr->Acceleration;
 	}
 	
 	UE_LOG(LogChaosMover, Warning, TEXT("Invalid acceleration on ChaosCharacterMoverComponent"));
diff --git a/Engine/Plugins/Experimental/ChaosMover/Source/ChaosMover/Private/Character/Modes/ChaosFallingMode.cpp b/Engine/Plugins/Experimental/ChaosMover/Source/ChaosMover/Private/Character/Modes/ChaosFallingMode.cpp
--- a/Engine/Plugins/Experimental/ChaosMover/Source/ChaosMover/Private/Character/Modes/ChaosFallingMode.cpp
+++ b/Engine/Plugins/Experimental/ChaosMover/Source/ChaosMover/Private/Character/Modes/ChaosFallingMode.cpp
@@ -61,6 +61,15 @@ void UChaosFallingMode::GenerateMove_Implementation(const FMoverTickStartData& S
 		return;
 	}
 
+	TStrongObjectPtr<const USharedChaosCharacterMovementSettings> SharedSettingsStrongPtr = SharedSettingsWeakPtr.Pin();
+	if (!SharedSettingsStrongPtr)
+	{
+		UE_LOG(LogChaosMover, Warning, TEXT("ChaosFallingMode unable to read from shared settings (USharedChaosCharacterMovementSettings)"));
+		return;
+	}
+	const USharedChaosCharacterMovementSettings* SharedSettings = SharedSettingsStrongPtr.Get();
+	check(SharedSettings);
+
 	const FMoverDefaultSyncState* StartingSyncState = StartState.SyncState.SyncStateCollection.FindDataByType<FMoverDefaultSyncState>();
 	check(StartingSyncState);
 
diff --git a/Engine/Plugins/Experimental/ChaosMover/Source/ChaosMover/Private/Character/Modes/ChaosFlyingMode.cpp b/Engine/Plugins/Experimental/ChaosMover/Source/ChaosMover/Private/Character/Modes/ChaosFlyingMode.cpp
--- a/Engine/Plugins/Experimental/ChaosMover/Source/ChaosMover/Private/Character/Modes/ChaosFlyingMode.cpp
+++ b/Engine/Plugins/Experimental/ChaosMover/Source/ChaosMover/Private/Character/Modes/ChaosFlyingMode.cpp
@@ -42,6 +42,15 @@ void UChaosFlyingMode::GenerateMove_Implementation(const FMoverTickStartData& St
 		return;
 	}
 
+	TStrongObjectPtr<const USharedChaosCharacterMovementSettings> SharedSettingsStrongPtr = SharedSettingsWeakPtr.Pin();
+	if (!SharedSettingsStrongPtr)
+	{
+		UE_LOG(LogChaosMover, Warning, TEXT("ChaosFlyingMode unable to read from shared settings (USharedChaosCharacterMovementSettings)"));
+		return;
+	}
+	const USharedChaosCharacterMovementSettings* SharedSettings = SharedSettingsStrongPtr.Get();
+	check(SharedSettings);
+
 	const FMoverDefaultSyncState* StartingSyncState = StartState.SyncState.SyncStateCollection.FindDataByType<FMoverDefaultSyncState>();
 	check(StartingSyncState);
 
diff --git a/Engine/Plugins/Experimental/ChaosMover/Source/ChaosMover/Private/Character/Modes/ChaosSwimmingMode.cpp b/Engine/Plugins/Experimental/ChaosMover/Source/ChaosMover/Private/Character/Modes/ChaosSwimmingMode.cpp
--- a/Engine/Plugins/Experimental/ChaosMover/Source/ChaosMover/Private/Character/Modes/ChaosSwimmingMode.cpp
+++ b/Engine/Plugins/Experimental/ChaosMover/Source/ChaosMover/Private/Character/Modes/ChaosSwimmingMode.cpp
@@ -51,7 +51,16 @@ void UChaosSwimmingMode::GenerateMove_Implementation(const FMoverTickStartData&
 	{
 		return;
 	}
-	
+
+	TStrongObjectPtr<const USharedChaosCharacterMovementSettings> SharedSettingsStrongPtr = SharedSettingsWeakPtr.Pin();
+	if (!SharedSettingsStrongPtr)
+	{
+		UE_LOG(LogChaosMover, Warning, TEXT("ChaosSwimmingMode unable to read from shared settings (USharedChaosCharacterMovementSettings)"));
+		return;
+	}
+	const USharedChaosCharacterMovementSettings* SharedSettings = SharedSettingsStrongPtr.Get();
+	check(SharedSettings);
+
 	const float DeltaSeconds = TimeStep.StepMs * 0.001f;	
 	const float HalfHeight = DefaultSimInputs->PawnCollisionHalfHeight;
 	FVector Velocity = StartingSyncState->GetVelocity_WorldSpace();
diff --git a/Engine/Plugins/Experimental/ChaosMover/Source/ChaosMover/Private/Character/Modes/ChaosWalkingMode.cpp b/Engine/Plugins/Experimental/ChaosMover/Source/ChaosMover/Private/Character/Modes/ChaosWalkingMode.cpp
--- a/Engine/Plugins/Experimental/ChaosMover/Source/ChaosMover/Private/Character/Modes/ChaosWalkingMode.cpp
+++ b/Engine/Plugins/Experimental/ChaosMover/Source/ChaosMover/Private/Character/Modes/ChaosWalkingMode.cpp
@@ -69,6 +69,15 @@ void UChaosWalkingMode::GenerateMove_Implementation(const FMoverTickStartData& S
 		return;
 	}
 
+	TStrongObjectPtr<const USharedChaosCharacterMovementSettings> SharedSettingsStrongPtr = SharedSettingsWeakPtr.Pin();
+	if (!SharedSettingsStrongPtr)
+	{
+		UE_LOG(LogChaosMover, Warning, TEXT("ChaosWalkingMode unable to read from shared settings (USharedChaosCharacterMovementSettings)"));
+		return;
+	}
+	const USharedChaosCharacterMovementSettings* SharedSettings = SharedSettingsStrongPtr.Get();
+	check(SharedSettings);
+
 	const FMoverDefaultSyncState* StartingSyncState = StartState.SyncState.SyncStateCollection.FindDataByType<FMoverDefaultSyncState>();
 	check(StartingSyncState);
 
@@ -229,6 +238,15 @@ void UChaosWalkingMode::SimulationTick_Implementation(const FSimulationTickParam
 		return;
 	}
 
+	TStrongObjectPtr<const USharedChaosCharacterMovementSettings> SharedSettingsStrongPtr = SharedSettingsWeakPtr.Pin();
+	if (!SharedSettingsStrongPtr)
+	{
+		UE_LOG(LogChaosMover, Warning, TEXT("UChaosWalkingMode unable to read from shared settings (USharedChaosCharacterMovementSettings)"));
+		return;
+	}
+	const USharedChaosCharacterMovementSettings* SharedSettings = SharedSettingsStrongPtr.Get();
+	check(SharedSettings);
+
 	FProposedMove ProposedMove = Params.ProposedMove;
 	const FMoverDefaultSyncState* StartingSyncState = Params.StartState.SyncState.SyncStateCollection.FindDataByType<FMoverDefaultSyncState>();
 	check(StartingSyncState);
@@ -359,6 +377,10 @@ bool UChaosWalkingMode::CanStepUpOnHitSurface(const FFloorCheckResult& FloorResu
 {
 	const float StepHeight = GetTargetHeight() - FloorResult.FloorDist;
 
+	TStrongObjectPtr<const USharedChaosCharacterMovementSettings> SharedSettingsStrongPtr = SharedSettingsWeakPtr.Pin();
+	const USharedChaosCharacterMovementSettings* SharedSettings = SharedSettingsStrongPtr ? SharedSettingsStrongPtr.Get() : GetDefault<USharedChaosCharacterMovementSettings>();
+	check(SharedSettings);
+
 	bool bWalkable = StepHeight <= SharedSettings->MaxStepHeight;
 	constexpr float MinStepHeight = 2.0f;
 	const bool SteppingUp = StepHeight > MinStepHeight;
@@ -397,6 +419,10 @@ void UChaosWalkingMode::GetFloorAndCheckMovement(const FMoverDefaultSyncState& S
 		VelocityPadding = FMath::Max(DefaultSimInputs.UpDir.Dot(SyncState.GetVelocity_WorldSpace()) * DeltaSeconds, 0.0f);
 	}
 
+	TStrongObjectPtr<const USharedChaosCharacterMovementSettings> SharedSettingsStrongPtr = SharedSettingsWeakPtr.Pin();
+	const USharedChaosCharacterMovementSettings* SharedSettings = SharedSettingsStrongPtr ? SharedSettingsStrongPtr.Get() : GetDefault<USharedChaosCharacterMovementSettings>();
+	check(SharedSettings);
+
 	UE::ChaosMover::Utils::FFloorSweepParams SweepParams{
 		.ResponseParams = DefaultSimInputs.CollisionResponseParams,
 		.QueryParams = DefaultSimInputs.CollisionQueryParams,
@@ -528,6 +554,15 @@ void UChaosWalkingMode::ModifyContacts(const FMoverTimeStep& TimeStep, const UE:
 		return;
 	}
 
+	TStrongObjectPtr<const USharedChaosCharacterMovementSettings> SharedSettingsStrongPtr = SharedSettingsWeakPtr.Pin();
+	if (!SharedSettingsStrongPtr)
+	{
+		UE_LOG(LogChaosMover, Warning, TEXT("UChaosWalkingMode unable to read from shared settings (USharedChaosCharacterMovementSettings)"));
+		return;
+	}
+	const USharedChaosCharacterMovementSettings* SharedSettings = SharedSettingsStrongPtr.Get();
+	check(SharedSettings);
+
 	// Try and find the ground particle if there is one in the latest floor result
 	Chaos::FGeometryParticleHandle* GroundParticle = nullptr;
 	FFloorCheckResult FloorResult;
diff --git a/Engine/Plugins/Experimental/ChaosMover/Source/ChaosMover/Private/Character/Transitions/ChaosCharacterFallingCheck.cpp b/Engine/Plugins/Experimental/ChaosMover/Source/ChaosMover/Private/Character/Transitions/ChaosCharacterFallingCheck.cpp
--- a/Engine/Plugins/Experimental/ChaosMover/Source/ChaosMover/Private/Character/Transitions/ChaosCharacterFallingCheck.cpp
+++ b/Engine/Plugins/Experimental/ChaosMover/Source/ChaosMover/Private/Character/Transitions/ChaosCharacterFallingCheck.cpp
@@ -28,15 +28,11 @@ void UChaosCharacterFallingCheck::OnRegistered()
 {
 	Super::OnRegistered();
 
-	SharedSettings = GetMoverComponent()->FindSharedSettings<USharedChaosCharacterMovementSettings>();
-	ensureMsgf(SharedSettings, TEXT("Failed to find instance of USharedChaosCharacterMovementSettings on %s. Movement may not function properly."), *GetPathNameSafe(this));
-}
-
-void UChaosCharacterFallingCheck::OnUnregistered()
-{
-	SharedSettings = nullptr;
-
-	Super::OnUnregistered();
+	TObjectPtr<const USharedChaosCharacterMovementSettings> SharedSettingsPtr = GetMoverComponent()->FindSharedSettings<USharedChaosCharacterMovementSettings>();
+	if (ensureMsgf(SharedSettingsPtr, TEXT("Failed to find instance of USharedChaosCharacterMovementSettings on %s. Movement may not function properly."), *GetPathNameSafe(this)))
+	{
+		SharedSettingsWeakPtr = TWeakObjectPtr<const USharedChaosCharacterMovementSettings>(SharedSettingsPtr);
+	}
 }
 
 FTransitionEvalResult UChaosCharacterFallingCheck::Evaluate_Implementation(const FSimulationTickParams& Params) const
@@ -70,6 +66,15 @@ FTransitionEvalResult UChaosCharacterFallingCheck::Evaluate_Implementation(const
 			return EvalResult;
 		}
 
+		TStrongObjectPtr<const USharedChaosCharacterMovementSettings> SharedSettingsStrongPtr = SharedSettingsWeakPtr.Pin();
+		if (!SharedSettingsStrongPtr)
+		{
+			UE_LOG(LogChaosMover, Warning, TEXT("UChaosCharacterFallingCheck unable to read from shared settings (USharedChaosCharacterMovementSettings)"));
+			return EvalResult;
+		}
+		const USharedChaosCharacterMovementSettings* SharedSettings = SharedSettingsStrongPtr.Get();
+		check(SharedSettings);
+
 		if (TStrongObjectPtr<const UBaseMovementMode> BaseMode = Simulation->FindMovementModeByName(Params.StartState.SyncState.MovementMode).Pin())
 		{
 			if (const IChaosCharacterMovementModeInterface* Mode = Cast<IChaosCharacterMovementModeInterface>(BaseMode.Get()))
diff --git a/Engine/Plugins/Experimental/ChaosMover/Source/ChaosMover/Public/ChaosMover/Character/Modes/ChaosCharacterMovementMode.h b/Engine/Plugins/Experimental/ChaosMover/Source/ChaosMover/Public/ChaosMover/Character/Modes/ChaosCharacterMovementMode.h
--- a/Engine/Plugins/Experimental/ChaosMover/Source/ChaosMover/Public/ChaosMover/Character/Modes/ChaosCharacterMovementMode.h
+++ b/Engine/Plugins/Experimental/ChaosMover/Source/ChaosMover/Public/ChaosMover/Character/Modes/ChaosCharacterMovementMode.h
@@ -34,7 +34,6 @@ public:
 	CHAOSMOVER_API UChaosCharacterMovementMode(const FObjectInitializer& ObjectInitializer);
 
 	CHAOSMOVER_API virtual void OnRegistered(const FName ModeName) override;
-	CHAOSMOVER_API virtual void OnUnregistered() override;
 
 	// IChaosCharacterConstraintMovementModeInterface
 	virtual bool ShouldEnableConstraint() const override
@@ -82,7 +81,7 @@ protected:
 	CHAOSMOVER_API void SetQueryRadiusOverride(float InQueryRadius);
 	CHAOSMOVER_API void ClearQueryRadiusOverride();
 
-	TObjectPtr<const class USharedChaosCharacterMovementSettings> SharedSettings;
+	TWeakObjectPtr<const class USharedChaosCharacterMovementSettings> SharedSettingsWeakPtr;
 
 	// Maximum force the character can apply to reach the motion target
 	UPROPERTY(EditAnywhere, Category = "Constraint Settings", meta = (ClampMin = "0", UIMin = "0", ForceUnits = "Newtons"))
diff --git a/Engine/Plugins/Experimental/ChaosMover/Source/ChaosMover/Public/ChaosMover/Character/Transitions/ChaosCharacterFallingCheck.h b/Engine/Plugins/Experimental/ChaosMover/Source/ChaosMover/Public/ChaosMover/Character/Transitions/ChaosCharacterFallingCheck.h
--- a/Engine/Plugins/Experimental/ChaosMover/Source/ChaosMover/Public/ChaosMover/Character/Transitions/ChaosCharacterFallingCheck.h
+++ b/Engine/Plugins/Experimental/ChaosMover/Source/ChaosMover/Public/ChaosMover/Character/Transitions/ChaosCharacterFallingCheck.h
@@ -18,7 +18,6 @@ public:
 	UE_API UChaosCharacterFallingCheck(const FObjectInitializer& ObjectInitializer);
 
 	UE_API virtual void OnRegistered() override;
-	UE_API virtual void OnUnregistered() override;
 
 	UE_API virtual FTransitionEvalResult Evaluate_Implementation(const FSimulationTickParams& Params) const override;
 	UE_API virtual void Trigger_Implementation(const FSimulationTickParams& Params) override;
@@ -34,7 +33,7 @@ public:
 	float MaxUnsupportedTimeBeforeFalling = 0.06f;
 
 protected:
-	TObjectPtr<const class USharedChaosCharacterMovementSettings> SharedSettings;
+	TWeakObjectPtr<const class USharedChaosCharacterMovementSettings> SharedSettingsWeakPtr;
 };
 
 #undef UE_API
