[Chaos Mover] Fix caclulation of inverse inertia when applying jump reaction impulse

#Synced-CL 50481757

diff --git a/Engine/Plugins/Experimental/ChaosMover/Source/ChaosMover/Private/Character/Transitions/ChaosCharacterJumpCheck.cpp b/Engine/Plugins/Experimental/ChaosMover/Source/ChaosMover/Private/Character/Transitions/ChaosCharacterJumpCheck.cpp
--- a/Engine/Plugins/Experimental/ChaosMover/Source/ChaosMover/Private/Character/Transitions/ChaosCharacterJumpCheck.cpp
+++ b/Engine/Plugins/Experimental/ChaosMover/Source/ChaosMover/Private/Character/Transitions/ChaosCharacterJumpCheck.cpp
@@ -4,6 +4,7 @@
 
 #include "Chaos/ParticleHandle.h"
 #include "Chaos/PhysicsObjectInternalInterface.h"
+#include "Chaos/Utilities.h"
 #include "ChaosMover/ChaosMoverLog.h"
 #include "ChaosMover/ChaosMoverSimulation.h"
 #include "ChaosMover/ChaosMoverSimulationTypes.h"
@@ -20,7 +21,8 @@ namespace Private
 	{
 		Chaos::FRigidTransform3 ComTransform = Particle->GetTransformXRCom();
 		const Chaos::FVec3 Offset = Location - ComTransform.GetLocation();
-		Particle->SetW(Particle->GetW() + Particle->InvI() * Offset.Cross(Impulse));
+		const Chaos::FMatrix33 WorldInvI = Chaos::Utilities::ComputeWorldSpaceInertia(ComTransform.GetRotation(), Particle->ConditionedInvI());
+		Particle->SetW(Particle->GetW() + WorldInvI * Offset.Cross(Impulse));
 		Particle->SetV(Particle->GetV() + Particle->InvM() * Impulse);
 	}
 }
