[Chaos Mover] Change assert on output sync state transform differing from particle transform into a log warning

#Synced-CL 50481757

diff --git a/Engine/Plugins/Experimental/ChaosMover/Source/ChaosMover/Private/Backends/ChaosMoverBackend.cpp b/Engine/Plugins/Experimental/ChaosMover/Source/ChaosMover/Private/Backends/ChaosMoverBackend.cpp
--- a/Engine/Plugins/Experimental/ChaosMover/Source/ChaosMover/Private/Backends/ChaosMoverBackend.cpp
+++ b/Engine/Plugins/Experimental/ChaosMover/Source/ChaosMover/Private/Backends/ChaosMoverBackend.cpp
@@ -669,7 +669,10 @@ void UChaosMoverBackendComponent::FinalizeFrame(double ResultsTimeInMs)
 			TransformSyncState.GetMovementBase(),
 			TransformSyncState.GetMovementBaseBoneName());
 
-		check(Chaos::FVec3::IsNearlyEqual(Chaos::FVec3(TransformSyncState.GetLocation_WorldSpace()), Particle->GetX(), UE_KINDA_SMALL_NUMBER));
+		if (!Chaos::FVec3::IsNearlyEqual(Chaos::FVec3(TransformSyncState.GetLocation_WorldSpace()), Particle->GetX(), UE_KINDA_SMALL_NUMBER))
+		{
+			UE_LOG(LogChaosMover, Warning, TEXT("SyncState location (%s) differs from particle location (%s)"), *TransformSyncState.GetLocation_WorldSpace().ToString(), *Particle->GetX().ToString());
+		}
 	}
 
 	MoverComp.SetSimulationOutput(TimeStep, InterpolatedOutput);
