[Chaos Mover] Add validation to modes and transitions to make sure they don't get added to an incompatible backend

#Synced-CL 50481757

diff --git a/Engine/Plugins/Experimental/ChaosMover/Source/ChaosMover/Private/ChaosMovementMode.cpp b/Engine/Plugins/Experimental/ChaosMover/Source/ChaosMover/Private/ChaosMovementMode.cpp
--- a/Engine/Plugins/Experimental/ChaosMover/Source/ChaosMover/Private/ChaosMovementMode.cpp
+++ b/Engine/Plugins/Experimental/ChaosMover/Source/ChaosMover/Private/ChaosMovementMode.cpp
@@ -4,9 +4,17 @@
 
 #include "ChaosMover/ChaosMovementModeTransition.h"
 #include "ChaosMover/ChaosMoverLog.h"
+#if WITH_EDITOR
+#include "ChaosMover/Backends/ChaosMoverBackend.h"
+#include "Internationalization/Text.h"
+#include "Misc/DataValidation.h"
+#include "MoverComponent.h"
+#endif // WITH_EDITOR
 
 #include UE_INLINE_GENERATED_CPP_BY_NAME(ChaosMovementMode)
 
+#define LOCTEXT_NAMESPACE "ChaosMovementMode"
+
 UChaosMovementMode::UChaosMovementMode(const FObjectInitializer& ObjectInitializer)
 	: Super(ObjectInitializer)
 	, Simulation(nullptr)
@@ -23,4 +31,25 @@ void UChaosMovementMode::SetSimulation(UChaosMoverSimulation* InSimulation)
 			ChaosTransition->SetSimulation(InSimulation);
 		}
 	}
-}
\ No newline at end of file
+}
+
+#if WITH_EDITOR
+EDataValidationResult UChaosMovementMode::IsDataValid(class FDataValidationContext& Context) const
+{
+	EDataValidationResult Result = Super::IsDataValid(Context);
+
+	if (const UMoverComponent* MoverComponent = GetMoverComponent())
+	{
+		const UClass* BackendClass = MoverComponent->BackendClass;
+		if (BackendClass && !BackendClass->IsChildOf<UChaosMoverBackendComponent>())
+		{
+			Context.AddError(NSLOCTEXT("ChaosMover", "ChaosMovementModeHasValidBackend", "Chaos movement modes need to have a backend class that supports Chaos Physics (UChaosMoverBackendComponent)."));
+			Result = EDataValidationResult::Invalid;
+		}
+	}
+
+	return Result;
+}
+#endif // WITH_EDITOR
+
+#undef LOCTEXT_NAMESPACE
\ No newline at end of file
diff --git a/Engine/Plugins/Experimental/ChaosMover/Source/ChaosMover/Private/ChaosMovementModeTransition.cpp b/Engine/Plugins/Experimental/ChaosMover/Source/ChaosMover/Private/ChaosMovementModeTransition.cpp
--- a/Engine/Plugins/Experimental/ChaosMover/Source/ChaosMover/Private/ChaosMovementModeTransition.cpp
+++ b/Engine/Plugins/Experimental/ChaosMover/Source/ChaosMover/Private/ChaosMovementModeTransition.cpp
@@ -3,11 +3,40 @@
 #include "ChaosMover/ChaosMovementModeTransition.h"
 
 #include "ChaosMover/ChaosMoverLog.h"
+#if WITH_EDITOR
+#include "ChaosMover/Backends/ChaosMoverBackend.h"
+#include "Internationalization/Text.h"
+#include "Misc/DataValidation.h"
+#include "MoverComponent.h"
+#endif // WITH_EDITOR
 
 #include UE_INLINE_GENERATED_CPP_BY_NAME(ChaosMovementModeTransition)
 
+#define LOCTEXT_NAMESPACE "ChaosMovementModeTransition"
+
 UChaosMovementModeTransition::UChaosMovementModeTransition(const FObjectInitializer& ObjectInitializer)
 	: Super(ObjectInitializer)
 	, Simulation(nullptr)
 {
-}
\ No newline at end of file
+}
+
+#if WITH_EDITOR
+EDataValidationResult UChaosMovementModeTransition::IsDataValid(class FDataValidationContext& Context) const
+{
+	EDataValidationResult Result = Super::IsDataValid(Context);
+	
+	if (const UMoverComponent* MoverComponent = GetMoverComponent())
+	{
+		const UClass* BackendClass = MoverComponent->BackendClass;
+		if (BackendClass && !BackendClass->IsChildOf<UChaosMoverBackendComponent>())
+		{
+			Context.AddError(NSLOCTEXT("ChaosMover", "ChaosMovementModeTransitionHasValidBackend", "Chaos movement mode transitions need to have a backend class that supports Chaos Physics (UChaosMoverBackendComponent)."));
+			Result = EDataValidationResult::Invalid;
+		}
+	}
+
+	return Result;
+}
+#endif // WITH_EDITOR
+
+#undef LOCTEXT_NAMESPACE
\ No newline at end of file
diff --git a/Engine/Plugins/Experimental/ChaosMover/Source/ChaosMover/Private/Character/Modes/ChaosFallingMode.cpp b/Engine/Plugins/Experimental/ChaosMover/Source/ChaosMover/Private/Character/Modes/ChaosFallingMode.cpp
--- a/Engine/Plugins/Experimental/ChaosMover/Source/ChaosMover/Private/Character/Modes/ChaosFallingMode.cpp
+++ b/Engine/Plugins/Experimental/ChaosMover/Source/ChaosMover/Private/Character/Modes/ChaosFallingMode.cpp
@@ -49,7 +49,7 @@ void UChaosFallingMode::GenerateMove_Implementation(const FMoverTickStartData& S
 
 	if (!Simulation)
 	{
-		UE_LOG(LogChaosMover, Warning, TEXT("No Simulation set on ChaosFallingMode"));
+		UE_LOG(LogChaosMover, Warning, TEXT("No Simulation set on UChaosFallingMode. Check you have a Chaos Backend"));
 		return;
 	}
 
@@ -221,7 +221,7 @@ void UChaosFallingMode::SimulationTick_Implementation(const FSimulationTickParam
 
 	if (!Simulation)
 	{
-		UE_LOG(LogChaosMover, Warning, TEXT("No Simulation set on ChaosFallingMode"));
+		UE_LOG(LogChaosMover, Warning, TEXT("No Simulation set on UChaosFallingMode. Check you have a Chaos Backend"));
 		return;
 	}
 
diff --git a/Engine/Plugins/Experimental/ChaosMover/Source/ChaosMover/Private/Character/Modes/ChaosFlyingMode.cpp b/Engine/Plugins/Experimental/ChaosMover/Source/ChaosMover/Private/Character/Modes/ChaosFlyingMode.cpp
--- a/Engine/Plugins/Experimental/ChaosMover/Source/ChaosMover/Private/Character/Modes/ChaosFlyingMode.cpp
+++ b/Engine/Plugins/Experimental/ChaosMover/Source/ChaosMover/Private/Character/Modes/ChaosFlyingMode.cpp
@@ -30,7 +30,7 @@ void UChaosFlyingMode::GenerateMove_Implementation(const FMoverTickStartData& St
 
 	if (!Simulation)
 	{
-		UE_LOG(LogChaosMover, Warning, TEXT("No Simulation set on ChaosFlyingMode"));
+		UE_LOG(LogChaosMover, Warning, TEXT("No Simulation set on UChaosFlyingMode. Check you have a Chaos Backend"));
 		return;
 	}
 
@@ -114,7 +114,7 @@ void UChaosFlyingMode::SimulationTick_Implementation(const FSimulationTickParams
 
 	if (!Simulation)
 	{
-		UE_LOG(LogChaosMover, Warning, TEXT("No Simulation set on ChaosFlyingMode"));
+		UE_LOG(LogChaosMover, Warning, TEXT("No Simulation set on UChaosFlyingMode. Check you have a Chaos Backend"));
 		return;
 	}
 
diff --git a/Engine/Plugins/Experimental/ChaosMover/Source/ChaosMover/Private/Character/Modes/ChaosSwimmingMode.cpp b/Engine/Plugins/Experimental/ChaosMover/Source/ChaosMover/Private/Character/Modes/ChaosSwimmingMode.cpp
--- a/Engine/Plugins/Experimental/ChaosMover/Source/ChaosMover/Private/Character/Modes/ChaosSwimmingMode.cpp
+++ b/Engine/Plugins/Experimental/ChaosMover/Source/ChaosMover/Private/Character/Modes/ChaosSwimmingMode.cpp
@@ -32,10 +32,12 @@ void UChaosSwimmingMode::GenerateMove_Implementation(const FMoverTickStartData&
 {
 	QUICK_SCOPE_CYCLE_COUNTER(ChaosSwimmingMode_GenerateMove);
 
-	if (!ensureMsgf(Simulation, TEXT("No Simulation set on ChaosSwimmingMode")))
+	if (!Simulation)
 	{
+		UE_LOG(LogChaosMover, Warning, TEXT("No Simulation set on UChaosSwimmingMode. Check you have a Chaos Backend"));
 		return;
 	}
+
 	const FChaosMoverSimulationDefaultInputs* DefaultSimInputs = Simulation->GetLocalSimInput().FindDataByType<FChaosMoverSimulationDefaultInputs>();
 	if (!ensureMsgf(DefaultSimInputs, TEXT("ChaosSwimmingMode requires FChaosMoverSimulationDefaultInputs")))
 	{
@@ -242,10 +244,12 @@ void UChaosSwimmingMode::SimulationTick_Implementation(const FSimulationTickPara
 {
 	QUICK_SCOPE_CYCLE_COUNTER(ChaosSwimmingMode_SimulationTick);
 
-	if (!ensureMsgf(Simulation, TEXT("No Simulation set on ChaosSwimmingMode")))
+	if (!Simulation)
 	{
+		UE_LOG(LogChaosMover, Warning, TEXT("No Simulation set on UChaosSwimmingMode. Check you have a Chaos Backend"));
 		return;
-	}	
+	}
+
 	const FChaosMoverSimulationDefaultInputs* DefaultSimInputs = Simulation->GetLocalSimInput().FindDataByType<FChaosMoverSimulationDefaultInputs>();
 	if (!ensureMsgf(DefaultSimInputs != nullptr, TEXT("No FChaosMoverSimulationDefaultInputs in UChaosSwimmingMode::SimulationTick_Implementation, skipping")))
 	{
diff --git a/Engine/Plugins/Experimental/ChaosMover/Source/ChaosMover/Private/Character/Modes/ChaosWalkingMode.cpp b/Engine/Plugins/Experimental/ChaosMover/Source/ChaosMover/Private/Character/Modes/ChaosWalkingMode.cpp
--- a/Engine/Plugins/Experimental/ChaosMover/Source/ChaosMover/Private/Character/Modes/ChaosWalkingMode.cpp
+++ b/Engine/Plugins/Experimental/ChaosMover/Source/ChaosMover/Private/Character/Modes/ChaosWalkingMode.cpp
@@ -57,7 +57,7 @@ void UChaosWalkingMode::GenerateMove_Implementation(const FMoverTickStartData& S
 
 	if (!Simulation)
 	{
-		UE_LOG(LogChaosMover, Warning, TEXT("No Simulation set on ChaosWalkingMode"));
+		UE_LOG(LogChaosMover, Warning, TEXT("No Simulation set on ChaosWalkingMode. Check you have a Chaos Backend"));
 		return;
 	}
 
@@ -226,7 +226,7 @@ void UChaosWalkingMode::SimulationTick_Implementation(const FSimulationTickParam
 
 	if (!Simulation)
 	{
-		UE_LOG(LogChaosMover, Warning, TEXT("No Simulation set on ChaosWalkingMode"));
+		UE_LOG(LogChaosMover, Warning, TEXT("No Simulation set on ChaosWalkingMode. Check you have a Chaos Backend"));
 		return;
 	}
 
diff --git a/Engine/Plugins/Experimental/ChaosMover/Source/ChaosMover/Private/Character/Transitions/ChaosCharacterCrouchCheck.cpp b/Engine/Plugins/Experimental/ChaosMover/Source/ChaosMover/Private/Character/Transitions/ChaosCharacterCrouchCheck.cpp
--- a/Engine/Plugins/Experimental/ChaosMover/Source/ChaosMover/Private/Character/Transitions/ChaosCharacterCrouchCheck.cpp
+++ b/Engine/Plugins/Experimental/ChaosMover/Source/ChaosMover/Private/Character/Transitions/ChaosCharacterCrouchCheck.cpp
@@ -110,7 +110,7 @@ FTransitionEvalResult UChaosCharacterCrouchCheck::Evaluate_Implementation(const
 
 	if (!Simulation)
 	{
-		UE_LOG(LogChaosMover, Warning, TEXT("No Simulation set on UChaosCharacterCrouchCheck"));
+		UE_LOG(LogChaosMover, Warning, TEXT("No Simulation set on UChaosCharacterCrouchCheck. Check you have a Chaos Backend"));
 		return EvalResult;
 	}
 
@@ -159,6 +159,12 @@ FTransitionEvalResult UChaosCharacterCrouchCheck::Evaluate_Implementation(const
 
 void UChaosCharacterCrouchCheck::Trigger_Implementation(const FSimulationTickParams& TickParams)
 {
+	if (!Simulation)
+	{
+		UE_LOG(LogChaosMover, Warning, TEXT("No Simulation set on UChaosCharacterCrouchCheck. Check you have a Chaos Backend"));
+		return;
+	}
+
 	if (bTriggerCrouch)
 	{
 		Crouch(TickParams);
diff --git a/Engine/Plugins/Experimental/ChaosMover/Source/ChaosMover/Private/Character/Transitions/ChaosCharacterFallingCheck.cpp b/Engine/Plugins/Experimental/ChaosMover/Source/ChaosMover/Private/Character/Transitions/ChaosCharacterFallingCheck.cpp
--- a/Engine/Plugins/Experimental/ChaosMover/Source/ChaosMover/Private/Character/Transitions/ChaosCharacterFallingCheck.cpp
+++ b/Engine/Plugins/Experimental/ChaosMover/Source/ChaosMover/Private/Character/Transitions/ChaosCharacterFallingCheck.cpp
@@ -41,10 +41,15 @@ FTransitionEvalResult UChaosCharacterFallingCheck::Evaluate_Implementation(const
 
 	if (TransitionToFallingMode.IsNone())
 	{
+		UE_LOG(LogChaosMover, Warning, TEXT("No TransitionToFallingMode set on UChaosCharacterFallingCheck."));
 		return EvalResult;
 	}
 	
-	check(Simulation);
+	if (!Simulation)
+	{
+		UE_LOG(LogChaosMover, Warning, TEXT("No Simulation set on UChaosCharacterFallingCheck. Check you have a Chaos Backend"));
+		return EvalResult;
+	}
 
 	const FMoverTickStartData& StartState = Params.StartState;
 
diff --git a/Engine/Plugins/Experimental/ChaosMover/Source/ChaosMover/Private/Character/Transitions/ChaosCharacterJumpCheck.cpp b/Engine/Plugins/Experimental/ChaosMover/Source/ChaosMover/Private/Character/Transitions/ChaosCharacterJumpCheck.cpp
--- a/Engine/Plugins/Experimental/ChaosMover/Source/ChaosMover/Private/Character/Transitions/ChaosCharacterJumpCheck.cpp
+++ b/Engine/Plugins/Experimental/ChaosMover/Source/ChaosMover/Private/Character/Transitions/ChaosCharacterJumpCheck.cpp
@@ -53,7 +53,7 @@ void UChaosCharacterJumpCheck::Trigger_Implementation(const FSimulationTickParam
 {
 	if (!Simulation)
 	{
-		UE_LOG(LogChaosMover, Warning, TEXT("No Simulation set on UChaosCharacterJumpCheck"));
+		UE_LOG(LogChaosMover, Warning, TEXT("No Simulation set on UChaosCharacterJumpCheck. Check you have a Chaos Backend"));
 		return;
 	}
 
diff --git a/Engine/Plugins/Experimental/ChaosMover/Source/ChaosMover/Private/Character/Transitions/ChaosCharacterLandingCheck.cpp b/Engine/Plugins/Experimental/ChaosMover/Source/ChaosMover/Private/Character/Transitions/ChaosCharacterLandingCheck.cpp
--- a/Engine/Plugins/Experimental/ChaosMover/Source/ChaosMover/Private/Character/Transitions/ChaosCharacterLandingCheck.cpp
+++ b/Engine/Plugins/Experimental/ChaosMover/Source/ChaosMover/Private/Character/Transitions/ChaosCharacterLandingCheck.cpp
@@ -23,10 +23,14 @@ UChaosCharacterLandingCheck::UChaosCharacterLandingCheck(const FObjectInitialize
 
 FTransitionEvalResult UChaosCharacterLandingCheck::Evaluate_Implementation(const FSimulationTickParams& Params) const
 {
-	check(Simulation);
-
 	FTransitionEvalResult EvalResult;
 
+	if (!Simulation)
+	{
+		UE_LOG(LogChaosMover, Warning, TEXT("No Simulation set on UChaosCharacterLandingCheck. Check you have a Chaos Backend"));
+		return EvalResult;
+	}
+
 	const FMoverTickStartData& StartState = Params.StartState;
 
 	if (const FChaosMoverSimulationDefaultInputs* SimInputs = Simulation->GetLocalSimInput().FindDataByType<FChaosMoverSimulationDefaultInputs>())
@@ -90,20 +94,23 @@ FTransitionEvalResult UChaosCharacterLandingCheck::Evaluate_Implementation(const
 
 void UChaosCharacterLandingCheck::Trigger_Implementation(const FSimulationTickParams& Params)
 {
+	if (!Simulation)
+	{
+		UE_LOG(LogChaosMover, Warning, TEXT("No Simulation set on UChaosCharacterLandingCheck. Check you have a Chaos Backend"));
+		return;
+	}
+
 	// Add a landed event to the simulation (will be broadcast on GT during post sim)
-	if (UChaosMoverSimulation* Sim = GetSimulation_Mutable())
+	if (const UMoverBlackboard* Blackboard = Simulation->GetBlackboard())
 	{
-		if (const UMoverBlackboard* Blackboard = Sim->GetBlackboard())
+		FFloorCheckResult FloorResult;
+		if (Blackboard->TryGet(CommonBlackboard::LastFloorResult, FloorResult))
 		{
-			FFloorCheckResult FloorResult;
-			if (Blackboard->TryGet(CommonBlackboard::LastFloorResult, FloorResult))
+			const TSharedPtr<FLandedEventData> LandedEventData = AllocLandedEventData();
+			if (LandedEventData.IsValid())
 			{
-				const TSharedPtr<FLandedEventData> LandedEventData = AllocLandedEventData();
-				if (LandedEventData.IsValid())
-				{
-					PopulateLandedEventData(Params, FloorResult, *LandedEventData);
-					Sim->AddEvent(LandedEventData);
-				}
+				PopulateLandedEventData(Params, FloorResult, *LandedEventData);
+				Simulation->AddEvent(LandedEventData);
 			}
 		}
 	}
diff --git a/Engine/Plugins/Experimental/ChaosMover/Source/ChaosMover/Private/Character/Transitions/ChaosCharacterLaunchCheck.cpp b/Engine/Plugins/Experimental/ChaosMover/Source/ChaosMover/Private/Character/Transitions/ChaosCharacterLaunchCheck.cpp
--- a/Engine/Plugins/Experimental/ChaosMover/Source/ChaosMover/Private/Character/Transitions/ChaosCharacterLaunchCheck.cpp
+++ b/Engine/Plugins/Experimental/ChaosMover/Source/ChaosMover/Private/Character/Transitions/ChaosCharacterLaunchCheck.cpp
@@ -46,7 +46,7 @@ void UChaosCharacterLaunchCheck::Trigger_Implementation(const FSimulationTickPar
 {
 	if (!Simulation)
 	{
-		UE_LOG(LogChaosMover, Warning, TEXT("No Simulation set on UChaosCharacterLaunchCheck"));
+		UE_LOG(LogChaosMover, Warning, TEXT("No Simulation set on UChaosCharacterLaunchCheck. Check you have a Chaos Backend"));
 		return;
 	}
 
diff --git a/Engine/Plugins/Experimental/ChaosMover/Source/ChaosMover/Private/Character/Transitions/ChaosCharacterWaterCheck.cpp b/Engine/Plugins/Experimental/ChaosMover/Source/ChaosMover/Private/Character/Transitions/ChaosCharacterWaterCheck.cpp
--- a/Engine/Plugins/Experimental/ChaosMover/Source/ChaosMover/Private/Character/Transitions/ChaosCharacterWaterCheck.cpp
+++ b/Engine/Plugins/Experimental/ChaosMover/Source/ChaosMover/Private/Character/Transitions/ChaosCharacterWaterCheck.cpp
@@ -23,10 +23,14 @@ UChaosCharacterWaterCheck::UChaosCharacterWaterCheck(const FObjectInitializer& O
 
 FTransitionEvalResult UChaosCharacterWaterCheck::Evaluate_Implementation(const FSimulationTickParams& Params) const
 {
-	check(Simulation);
-
 	FTransitionEvalResult EvalResult;
 
+	if (!Simulation)
+	{
+		UE_LOG(LogChaosMover, Warning, TEXT("No Simulation set on UChaosCharacterWaterCheck. Check you have a Chaos Backend"));
+		return EvalResult;
+	}
+
 	const FMoverTickStartData& StartState = Params.StartState;
 
 	if (const FChaosMoverSimulationDefaultInputs* SimInputs = Simulation->GetLocalSimInput().FindDataByType<FChaosMoverSimulationDefaultInputs>())
diff --git a/Engine/Plugins/Experimental/ChaosMover/Source/ChaosMover/Public/ChaosMover/ChaosMovementMode.h b/Engine/Plugins/Experimental/ChaosMover/Source/ChaosMover/Public/ChaosMover/ChaosMovementMode.h
--- a/Engine/Plugins/Experimental/ChaosMover/Source/ChaosMover/Public/ChaosMover/ChaosMovementMode.h
+++ b/Engine/Plugins/Experimental/ChaosMover/Source/ChaosMover/Public/ChaosMover/ChaosMovementMode.h
@@ -7,8 +7,6 @@
 
 #include "ChaosMovementMode.generated.h"
 
-#define UE_API CHAOSMOVER_API
-
 class UChaosMoverSimulation;
 namespace Chaos
 {
@@ -31,7 +29,7 @@ class UChaosMovementMode : public UBaseMovementMode
 	GENERATED_BODY()
 
 public:
-	UE_API UChaosMovementMode(const FObjectInitializer& ObjectInitializer);
+	CHAOSMOVER_API UChaosMovementMode(const FObjectInitializer& ObjectInitializer);
 
 	UFUNCTION(BlueprintPure, Category = ChaosMover)
 	const UChaosMoverSimulation* GetSimulation() const
@@ -39,7 +37,7 @@ public:
 		return Simulation;
 	}
 
-	UE_API void SetSimulation(UChaosMoverSimulation* InSimulation);
+	CHAOSMOVER_API void SetSimulation(UChaosMoverSimulation* InSimulation);
 
 	// Whether this movement mode is relative to a basis transform
 	virtual bool UsesMovementBasisTransform() const
@@ -61,6 +59,10 @@ public:
 	UPROPERTY(EditAnywhere, Category = "Collision Settings")
 	EChaosMoverIgnoredCollisionMode IgnoredCollisionMode = EChaosMoverIgnoredCollisionMode::DisableCollisionsWithIgnored;
 
+#if WITH_EDITOR
+	CHAOSMOVER_API virtual EDataValidationResult IsDataValid(class FDataValidationContext& Context) const override;
+#endif // WITH_EDITOR
+
 protected:
 	UChaosMoverSimulation* Simulation;
 };
diff --git a/Engine/Plugins/Experimental/ChaosMover/Source/ChaosMover/Public/ChaosMover/ChaosMovementModeTransition.h b/Engine/Plugins/Experimental/ChaosMover/Source/ChaosMover/Public/ChaosMover/ChaosMovementModeTransition.h
--- a/Engine/Plugins/Experimental/ChaosMover/Source/ChaosMover/Public/ChaosMover/ChaosMovementModeTransition.h
+++ b/Engine/Plugins/Experimental/ChaosMover/Source/ChaosMover/Public/ChaosMover/ChaosMovementModeTransition.h
@@ -36,6 +36,10 @@ public:
 		Simulation = InSimulation;
 	}
 
+#if WITH_EDITOR
+	CHAOSMOVER_API virtual EDataValidationResult IsDataValid(class FDataValidationContext& Context) const override;
+#endif // WITH_EDITOR
+
 protected:
 	UChaosMoverSimulation* Simulation;
 };
\ No newline at end of file
