ChaosMover: Fix a couple of bugs in how async output is consumed and interpolated:
- We only pop async outputs which start time is <= to ResultsTime (that's the physics look back time, a time at a fixed delay behind game time at which we intepolate physics output)
- We convert ResultsTime to a "base time", i.e. the time of the beginning of a async frame step, so it can be properly interpolated on the basis of the async outputs' base times

#Synced-CL 50489768

diff --git a/Engine/Plugins/Experimental/ChaosMover/Source/ChaosMover/Private/Backends/ChaosMoverBackend.cpp b/Engine/Plugins/Experimental/ChaosMover/Source/ChaosMover/Private/Backends/ChaosMoverBackend.cpp
--- a/Engine/Plugins/Experimental/ChaosMover/Source/ChaosMover/Private/Backends/ChaosMoverBackend.cpp
+++ b/Engine/Plugins/Experimental/ChaosMover/Source/ChaosMover/Private/Backends/ChaosMoverBackend.cpp
@@ -650,11 +650,19 @@ void UChaosMoverBackendComponent::FinalizeFrame(double ResultsTimeInMs)
 		return;
 	}
 
+	Chaos::FPhysicsSolver* Solver = GetPhysicsSolver();
+	if (!Solver)
+	{
+		return;
+	}
+
 	UMoverComponent& MoverComp = GetMoverComponent();
 
 	FMoverTimeStep TimeStep;
 	UE::ChaosMover::FSimulationOutputData InterpolatedOutput;
-	SimOutputRecord.CreateInterpolatedResult(ResultsTimeInMs, TimeStep, InterpolatedOutput);
+	// ResultsTimeInMs is the time at the end of the step while CreateInterpolatedResult expects the time at the beginning of the step so we correct by subtracting one async dt
+	const double AtBaseTimeMs = ResultsTimeInMs - Solver->GetAsyncDeltaTime() * 1000.0;
+	SimOutputRecord.CreateInterpolatedResult(AtBaseTimeMs, /*OUT*/ TimeStep, /*OUT*/ InterpolatedOutput);
 
 	// Physics interactions in the last frame may have caused a change in position or velocity that's different from what a simple lerp would predict,
 	// so stomp the lerped sync state's transform data with that of the actual particle after the last sim frame
diff --git a/Engine/Plugins/Experimental/ChaosMover/Source/ChaosMover/Private/Backends/ChaosMoverSubsystem.cpp b/Engine/Plugins/Experimental/ChaosMover/Source/ChaosMover/Private/Backends/ChaosMoverSubsystem.cpp
--- a/Engine/Plugins/Experimental/ChaosMover/Source/ChaosMover/Private/Backends/ChaosMoverSubsystem.cpp
+++ b/Engine/Plugins/Experimental/ChaosMover/Source/ChaosMover/Private/Backends/ChaosMoverSubsystem.cpp
@@ -198,8 +198,10 @@ void UChaosMoverSubsystem::OnPostPhysicsTick(FChaosScene* Scene)
 
 	const bool ForceSingleThread = UE::ChaosMover::CVars::bForceSingleThreadedGT;
 
-	// Pop each async output
-	while (Chaos::TSimCallbackOutputHandle<UE::ChaosMover::FAsyncCallbackOutput> AsyncOutput = AsyncCallback->PopFutureOutputData_External())
+	// Pop each async output which START TIME is <= ResultsTimeInMs
+	// Since they are ordered, this ensures that we will pop the last output which end time is greater than ResultsTimeInMs, if there is such an output available,
+	// but not the ones after that, which will remain in the queue
+	while (Chaos::TSimCallbackOutputHandle<UE::ChaosMover::FAsyncCallbackOutput> AsyncOutput = AsyncCallback->PopOutputData_External())
 	{
 		auto LambdaParallelUpdate = [&AsyncOutput](int32 Idx) {
 			if (UChaosMoverBackendComponent* Backend = AsyncOutput->Backends[Idx].Get())
@@ -243,4 +245,4 @@ int32 UChaosMoverSubsystem::GetNetworkPhysicsTickOffset() const
 	}
 
 	return 0;
-}
\ No newline at end of file
+}
