[ChaosMover] BUGFIX - Cache SimBlackBoard's LastWaterResult on ChaosMoverComponent similar to LastFloorResult.
Also rolled back a change that merges MoverComponent's blackboard with MoverSimulation's blackboard.

#Synced-CL 50876330

#Synced-CL 50943561

diff --git a/Engine/Plugins/Experimental/ChaosMover/Source/ChaosMover/Private/ChaosMoverSimulation.cpp b/Engine/Plugins/Experimental/ChaosMover/Source/ChaosMover/Private/ChaosMoverSimulation.cpp
--- a/Engine/Plugins/Experimental/ChaosMover/Source/ChaosMover/Private/ChaosMoverSimulation.cpp
+++ b/Engine/Plugins/Experimental/ChaosMover/Source/ChaosMover/Private/ChaosMoverSimulation.cpp
@@ -564,6 +564,13 @@ void UChaosMoverSimulation::PostSimulationTickCharacter(const IChaosCharacterMov
 
 	FFloorResultData& FloorData = OutputData.AdditionalOutputData.FindOrAddMutableDataByType<FFloorResultData>();
 	FloorData.FloorResult = FloorResult;
+
+	// Update the water result
+	FWaterCheckResult WaterResult;
+	Blackboard->TryGet(CommonBlackboard::LastWaterResult, WaterResult);
+
+	FChaosWaterResultData& WaterData = OutputData.AdditionalOutputData.FindOrAddMutableDataByType<FChaosWaterResultData>();
+	WaterData.WaterResult = WaterResult;
 }
 
 void UChaosMoverSimulation::PostSimulationTickCharacterConstraint(const IChaosCharacterConstraintMovementModeInterface& CharacterConstraintMode, const FMoverTimeStep& TimeStep, UE::ChaosMover::FSimulationOutputData& OutputData)
diff --git a/Engine/Plugins/Experimental/ChaosMover/Source/ChaosMover/Private/ChaosMoverSimulationTypes.cpp b/Engine/Plugins/Experimental/ChaosMover/Source/ChaosMover/Private/ChaosMoverSimulationTypes.cpp
--- a/Engine/Plugins/Experimental/ChaosMover/Source/ChaosMover/Private/ChaosMoverSimulationTypes.cpp
+++ b/Engine/Plugins/Experimental/ChaosMover/Source/ChaosMover/Private/ChaosMoverSimulationTypes.cpp
@@ -157,6 +157,35 @@ UScriptStruct* FNetworkPhysicsDebugData::GetScriptStruct() const
 	return StaticStruct();
 }
 
+void FChaosWaterResultData::ToString(FAnsiStringBuilderBase& Out) const
+{
+	Super::ToString(Out);
+
+	Out.Appendf("bSwimmableVolume: %i | ", WaterResult.bSwimmableVolume);
+	Out.Appendf("ImmersionDepth: %.2f | ", WaterResult.WaterSplineData.ImmersionDepth);
+	Out.Appendf("ImmersionPercent: %.2f | ", WaterResult.WaterSplineData.ImmersionPercent);
+	Out.Appendf("WaterDepth: %.2f/n", WaterResult.WaterSplineData.WaterDepth);
+	Out.Appendf("HitResult: %s/n", *WaterResult.HitResult.ToString());
+}
+
+bool FChaosWaterResultData::ShouldReconcile(const FMoverDataStructBase& AuthorityState) const
+{
+	return false;
+}
+
+void FChaosWaterResultData::Interpolate(const FMoverDataStructBase& From, const FMoverDataStructBase& To, float Pct)
+{
+	*this = static_cast<const FChaosWaterResultData&>((Pct < 0.5f) ? From : To);
+}
+
+void FChaosWaterResultData::Merge(const FMoverDataStructBase& From)
+{
+}
+
+void FChaosWaterResultData::Decay(float DecayAmount)
+{
+}
+
 void FChaosNetInstantMovementEffectsQueue::Add(const FScheduledInstantMovementEffect& InScheduledEffect, int32 IssuanceServerFrame, uint8 UniqueID)
 {
 	TCheckedObjPtr<FInstantMovementEffect> InstantMovementEffect = InScheduledEffect.Effect.Get();
diff --git a/Engine/Plugins/Experimental/ChaosMover/Source/ChaosMover/Private/Character/ChaosCharacterMoverComponent.cpp b/Engine/Plugins/Experimental/ChaosMover/Source/ChaosMover/Private/Character/ChaosCharacterMoverComponent.cpp
--- a/Engine/Plugins/Experimental/ChaosMover/Source/ChaosMover/Private/Character/ChaosCharacterMoverComponent.cpp
+++ b/Engine/Plugins/Experimental/ChaosMover/Source/ChaosMover/Private/Character/ChaosCharacterMoverComponent.cpp
@@ -8,6 +8,7 @@
 #include "ChaosMover/Character/Modes/ChaosFlyingMode.h"
 #include "ChaosMover/Character/Modes/ChaosWalkingMode.h"
 #include "ChaosMover/ChaosMoverLog.h"
+#include "ChaosMover/ChaosMoverSimulationTypes.h"
 #include "DefaultMovementSet/CharacterMoverSimulationTypes.h"
 #include "GameFramework/Pawn.h"
 
@@ -82,6 +83,22 @@ void UChaosCharacterMoverComponent::SetAdditionalSimulationOutput(const FMoverDa
 		bFloorResultSet = true;
 		LatestFloorResult = FloorData->FloorResult;
 	}
+
+	if (const FChaosWaterResultData* WaterData = Data.FindDataByType<FChaosWaterResultData>())
+	{
+		bWaterResultSet = true;
+		LatestWaterResult = WaterData->WaterResult;
+	}
+}
+
+bool UChaosCharacterMoverComponent::TryGetLastWaterResult(FWaterCheckResult& OutWaterResult) const
+{
+	if (bWaterResultSet)
+	{
+		OutWaterResult = LatestWaterResult;
+		return true;
+	}
+	return false;
 }
 
 bool UChaosCharacterMoverComponent::TryGetFloorCheckHitResult(FHitResult& OutHitResult) const
diff --git a/Engine/Plugins/Experimental/ChaosMover/Source/ChaosMover/Public/ChaosMover/ChaosMoverSimulationTypes.h b/Engine/Plugins/Experimental/ChaosMover/Source/ChaosMover/Public/ChaosMover/ChaosMoverSimulationTypes.h
--- a/Engine/Plugins/Experimental/ChaosMover/Source/ChaosMover/Public/ChaosMover/ChaosMoverSimulationTypes.h
+++ b/Engine/Plugins/Experimental/ChaosMover/Source/ChaosMover/Public/ChaosMover/ChaosMoverSimulationTypes.h
@@ -256,6 +256,35 @@ struct FNetworkPhysicsDebugData : public FMoverDataStructBase
 	float CurrentInputDecay = -1.0f;
 };
 
+USTRUCT(BlueprintType)
+struct FChaosWaterResultData : public FMoverDataStructBase
+{
+	GENERATED_BODY()
+
+public:
+	UPROPERTY(BlueprintReadWrite, Category = Mover)
+	FWaterCheckResult WaterResult;
+
+	FChaosWaterResultData() = default;
+	virtual ~FChaosWaterResultData() = default;
+
+	virtual FMoverDataStructBase* Clone() const override
+	{
+		return new FChaosWaterResultData(*this);
+	}
+
+	virtual UScriptStruct* GetScriptStruct() const override
+	{
+		return StaticStruct();
+	}
+
+	virtual void ToString(FAnsiStringBuilderBase& Out) const override;
+	virtual bool ShouldReconcile(const FMoverDataStructBase& AuthorityState) const override;
+	virtual void Interpolate(const FMoverDataStructBase& From, const FMoverDataStructBase& To, float Pct) override;
+	virtual void Merge(const FMoverDataStructBase& From) override;
+	virtual void Decay(float DecayAmount) override;
+};
+
 UINTERFACE(BlueprintType, MinimalAPI, meta = (CannotImplementInterfaceInBlueprint))
 class UChaosCharacterMovementModeInterface : public UInterface
 {
diff --git a/Engine/Plugins/Experimental/ChaosMover/Source/ChaosMover/Public/ChaosMover/Character/ChaosCharacterMoverComponent.h b/Engine/Plugins/Experimental/ChaosMover/Source/ChaosMover/Public/ChaosMover/Character/ChaosCharacterMoverComponent.h
--- a/Engine/Plugins/Experimental/ChaosMover/Source/ChaosMover/Public/ChaosMover/Character/ChaosCharacterMoverComponent.h
+++ b/Engine/Plugins/Experimental/ChaosMover/Source/ChaosMover/Public/ChaosMover/Character/ChaosCharacterMoverComponent.h
@@ -25,6 +25,10 @@ public:
 	CHAOSMOVER_API UChaosCharacterMoverComponent();
 
 	CHAOSMOVER_API virtual bool TryGetFloorCheckHitResult(FHitResult& OutHitResult) const override;
+
+	UFUNCTION(BlueprintPure, Category = "Chaos Mover")
+	CHAOSMOVER_API virtual bool TryGetLastWaterResult(FWaterCheckResult& OutWaterResult) const;
+
 	CHAOSMOVER_API virtual void ProduceInput(const int32 DeltaTimeMS, FMoverInputCmdContext* Cmd) override;
 	CHAOSMOVER_API virtual void ProduceServerInput(const int32 DeltaTimeMS, FMoverInputCmdContext* Cmd);
 	CHAOSMOVER_API virtual void DoQueueNextMode(FName DesiredModeName, bool bShouldReenter=false) override;
@@ -77,6 +81,9 @@ protected:
 	bool bFloorResultSet = false;
 	FFloorCheckResult LatestFloorResult;
 
+	bool bWaterResultSet = false;
+	FWaterCheckResult LatestWaterResult;
+
 	FVector LaunchVelocityOrImpulse = FVector::ZeroVector;
 	EChaosMoverVelocityEffectMode LaunchMode = EChaosMoverVelocityEffectMode::AdditiveVelocity;
 
