Mover: Fix for NonUnity build errors

#Synced-CL 51010646

diff --git a/Engine/Plugins/Experimental/ChaosMover/Source/ChaosMover/Private/Character/Modes/ChaosCharacterMovementMode.cpp b/Engine/Plugins/Experimental/ChaosMover/Source/ChaosMover/Private/Character/Modes/ChaosCharacterMovementMode.cpp
--- a/Engine/Plugins/Experimental/ChaosMover/Source/ChaosMover/Private/Character/Modes/ChaosCharacterMovementMode.cpp
+++ b/Engine/Plugins/Experimental/ChaosMover/Source/ChaosMover/Private/Character/Modes/ChaosCharacterMovementMode.cpp
@@ -16,6 +16,7 @@
 #include "MoveLibrary/FloorQueryUtils.h"
 #include "MoveLibrary/MovementUtils.h"
 #include "MoveLibrary/WaterMovementUtils.h"
+#include "MoverComponent.h"
 
 #include UE_INLINE_GENERATED_CPP_BY_NAME(ChaosCharacterMovementMode)
 
