Mover:
- Adding support for scheduling layered moves.
- Adding issuance frame for scheduled layered moves and instant movement effects.
- Correcting rollback to only roll back the "mover commands" (effects, moves, etc.) which issuance time is after the new time we roll back to.
- This CL does not include networking of moves and effects for Mover, so scheduling is not really meanigful yet but will become important once networked
- Added FMoveTime a struct that either contains a FrameCount or a floating point TimeMs
- MoverBackendLiaisonInterface: Added IsFixedDt() and GetScheduledNetworkTime() functions to centralize the implementation of network scheduling to the backend
- Added MoverSchedulingUtils.h/.cpp to centralize network scheduling functions
Chaos Mover:
- Support for layered moves (networked like Instant Movement Effects).

The layered move scheduling/rollback code is pretty much a duplication of the code to support instant movement effects. This shows that we would benefit from an abstraction of sim or mover commands (queueing of an instant movement effects, of a layered move queueing, of a modifier, etc.) to factorize the code repeated here. This can happen in a future refactor.

#Synced-CL 51387699

diff --git a/Engine/Plugins/Experimental/ChaosMover/Source/ChaosMover/Private/Backends/ChaosMoverBackend.cpp b/Engine/Plugins/Experimental/ChaosMover/Source/ChaosMover/Private/Backends/ChaosMoverBackend.cpp
--- a/Engine/Plugins/Experimental/ChaosMover/Source/ChaosMover/Private/Backends/ChaosMoverBackend.cpp
+++ b/Engine/Plugins/Experimental/ChaosMover/Source/ChaosMover/Private/Backends/ChaosMoverBackend.cpp
@@ -425,7 +425,7 @@ int32 UChaosMoverBackendComponent::GetCurrentSimFrame()
 	// Note: this is implicitly an _External function
 	if (UWorld* World = GetWorld())
 	{
-		UE::NetworkPhysicsUtils::GetUpcomingServerFrame_External(World);
+		return UE::NetworkPhysicsUtils::GetUpcomingServerFrame_External(World);
 	}
 	return 0;
 }
@@ -447,11 +447,31 @@ float UChaosMoverBackendComponent::GetEventSchedulingMinDelaySeconds() const
 	return EventSchedulingMinDelaySeconds;
 }
 
+FMoverTime UChaosMoverBackendComponent::GetScheduledNetworkTime(const FMoverTime& Time) const
+{
+	UWorld* World = GetWorld();
+	FPhysScene* Scene = World ? World->GetPhysicsScene() : nullptr;
+	if (Chaos::FPhysicsSolver* Solver = Scene ? Scene->GetSolver() : nullptr)
+	{
+		const double DeltaTime = Solver->GetAsyncDeltaTime();
+		const int32 DelayFrames = (DeltaTime > UE_SMALL_NUMBER) ? FMath::CeilToInt32(GetEventSchedulingMinDelaySeconds() / DeltaTime) : 0;
+		const int32 ScheduledFrame = Time.FrameCount + DelayFrames;
+
+		return FMoverTime(ScheduledFrame, ScheduledFrame * DeltaTime * 1000);
+	}
+	return IMoverBackendLiaisonInterface::GetScheduledNetworkTime(Time);
+}
+
 bool UChaosMoverBackendComponent::IsAsync() const
 {
 	return true;
 }
 
+bool UChaosMoverBackendComponent::IsFixedDt() const
+{
+	return true;
+}
+
 bool UChaosMoverBackendComponent::ShouldResim() const
 {
 	return Simulation ? Simulation->ShouldResim() : IMoverBackendLiaisonInterface::ShouldResim();
@@ -503,30 +523,9 @@ void UChaosMoverBackendComponent::ProduceInputData(int32 PhysicsStep, int32 NumS
 	}
 
 	UMoverComponent& MoverComp = GetMoverComponent();
-	// Add all queued instant effects to the input so it can be networked
-	const TArray<FScheduledInstantMovementEffect>& ScheduledInstantMovementEffects = MoverComp.GetQueuedInstantMovementEffects();
-	if (!ScheduledInstantMovementEffects.IsEmpty())
-	{
-		if (bIsGeneratingInputsLocally)
-		{
-			// Input generating simulations inject instant movement effects from the Game Thread into the simulation input collection
-			// It will be used directly by the state machine (same as when it is received from the network or overwritten during resim),
-			// and will also be networked via NetInputCmd
-			InjectInstantMovementEffectsIntoInput(TimeStep, InputData);
-		}
-		else
-		{
-			// Transfer queued instant movement effects to the simulation
-			// Instant effects should still be consumed by a non input generating actor on a server if it is not controlled remotely (this is an approximation of "no other instance is input producing")
-			InjectInstantMovementEffectsIntoSim(TimeStep);
-		}
-		// We clear queued instant movement effects systematically to avoid accumulation, even if they ended up being discarded
-		MoverComp.ClearQueuedInstantMovementEffects();
-	}
 
 	// Add default simulation input data
 	FChaosMoverSimulationDefaultInputs& SimInputs = Simulation->GetLocalSimInput_Mutable().FindOrAddMutableDataByType<FChaosMoverSimulationDefaultInputs>();
-
 	SimInputs.Gravity = MoverComp.GetGravityAcceleration();
 	SimInputs.UpDir = MoverComp.GetUpDirection();
 	SimInputs.bIsGeneratingInputsLocally = bIsGeneratingInputsLocally;
@@ -536,30 +535,39 @@ void UChaosMoverBackendComponent::ProduceInputData(int32 PhysicsStep, int32 NumS
 	SimInputs.World = GetWorld();
 	SimInputs.AsyncNetworkPhysicsComponent = NetworkPhysicsComponent ? NetworkPhysicsComponent->GetNetworkPhysicsComponent_Internal() : nullptr;
 
-	if (UPrimitiveComponent* PrimComp = Cast<UPrimitiveComponent>(MoverComp.GetUpdatedComponent()))
-	{
-		SimInputs.CollisionQueryParams = FCollisionQueryParams(SCENE_QUERY_STAT(ChaosMoverQuery), false, PrimComp->GetOwner());
-		SimInputs.CollisionQueryParams.bTraceIntoSubComponents = false;
-		SimInputs.CollisionResponseParams = FCollisionResponseParams(ECR_Overlap);
-		SimInputs.CollisionResponseParams.CollisionResponse.SetResponse(ECC_WorldStatic, ECR_Block);
-		SimInputs.CollisionResponseParams.CollisionResponse.SetResponse(ECC_WorldDynamic, ECR_Block);
-		SimInputs.CollisionResponseParams.CollisionResponse.SetResponse(ECC_Vehicle, ECR_Block);
-		SimInputs.CollisionResponseParams.CollisionResponse.SetResponse(ECC_Destructible, ECR_Block);
-		SimInputs.CollisionResponseParams.CollisionResponse.SetResponse(ECC_PhysicsBody, ECR_Block);
-		PrimComp->InitSweepCollisionParams(SimInputs.CollisionQueryParams, SimInputs.CollisionResponseParams);
-
-		SimInputs.CollisionChannel = PrimComp->GetCollisionObjectType();
-		PrimComp->CalcBoundingCylinder(SimInputs.PawnCollisionRadius, SimInputs.PawnCollisionHalfHeight);
-	}
-	if (IPhysicsComponent* PhysComp = Cast<IPhysicsComponent>(MoverComp.GetUpdatedComponent()))
+	if (USceneComponent* UpdatedComponent = MoverComp.GetUpdatedComponent())
 	{
-		SimInputs.PhysicsObject = PhysComp->GetPhysicsObjectById(0); // Get the root physics object
-	}
-	if (const APhysicsVolume* CurPhysVolume = MoverComp.GetUpdatedComponent()->GetPhysicsVolume())
-	{
-		SimInputs.PhysicsObjectGravity = CurPhysVolume->GetGravityZ();
+		if (UPrimitiveComponent* PrimComp = Cast<UPrimitiveComponent>(UpdatedComponent))
+		{
+			SimInputs.CollisionQueryParams = FCollisionQueryParams(SCENE_QUERY_STAT(ChaosMoverQuery), false, PrimComp->GetOwner());
+			SimInputs.CollisionQueryParams.bTraceIntoSubComponents = false;
+			SimInputs.CollisionResponseParams = FCollisionResponseParams(ECR_Overlap);
+			SimInputs.CollisionResponseParams.CollisionResponse.SetResponse(ECC_WorldStatic, ECR_Block);
+			SimInputs.CollisionResponseParams.CollisionResponse.SetResponse(ECC_WorldDynamic, ECR_Block);
+			SimInputs.CollisionResponseParams.CollisionResponse.SetResponse(ECC_Vehicle, ECR_Block);
+			SimInputs.CollisionResponseParams.CollisionResponse.SetResponse(ECC_Destructible, ECR_Block);
+			SimInputs.CollisionResponseParams.CollisionResponse.SetResponse(ECC_PhysicsBody, ECR_Block);
+			PrimComp->InitSweepCollisionParams(SimInputs.CollisionQueryParams, SimInputs.CollisionResponseParams);
+
+			SimInputs.CollisionChannel = PrimComp->GetCollisionObjectType();
+			PrimComp->CalcBoundingCylinder(SimInputs.PawnCollisionRadius, SimInputs.PawnCollisionHalfHeight);
+		}
+		if (IPhysicsComponent* PhysComp = Cast<IPhysicsComponent>(UpdatedComponent))
+		{
+			SimInputs.PhysicsObject = PhysComp->GetPhysicsObjectById(0); // Get the root physics object
+		}
+		if (const APhysicsVolume* CurPhysVolume = UpdatedComponent->GetPhysicsVolume())
+		{
+			SimInputs.PhysicsObjectGravity = CurPhysVolume->GetGravityZ();
+		}
 	}
 
+	// TODO: Generalize this to sim commands so we can support this mechanism for IMEs, layered moves, modifiers, in one single queue. Also each command could be networked or not (the local ones don't get marshalled via input)
+	// and scheduled or not
+	// This needs to happen after local sim inputs because some of the logic reads from local sim inputs
+	InjectInstantMovementEffectsIntoSim(InputData);
+	InjectLayeredMovesIntoSim(InputData);
+
 	if (MoverComp.OnPreSimulationTick.IsBound())
 	{
 		MoverComp.OnPreSimulationTick.Broadcast(TimeStep, InputData.InputCmd);
@@ -582,49 +590,158 @@ void UChaosMoverBackendComponent::GenerateServerInput(const FMoverTimeStep& Time
 	}
 }
 
-void UChaosMoverBackendComponent::InjectInstantMovementEffectsIntoInput(const FMoverTimeStep& TimeStep, UE::ChaosMover::FSimulationInputData& InputData)
+void UChaosMoverBackendComponent::InjectInstantMovementEffectsIntoSim(UE::ChaosMover::FSimulationInputData& OutInputData)
 {
+	const bool bIsGeneratingInputsLocally = !NetworkPhysicsComponent || NetworkPhysicsComponent->IsLocallyControlled();
+
 	UMoverComponent& MoverComp = GetMoverComponent();
+
 	// Add all queued instant effects to the input so it can be networked
 	const TArray<FScheduledInstantMovementEffect>& ScheduledInstantMovementEffects = MoverComp.GetQueuedInstantMovementEffects();
-	FChaosNetInstantMovementEffectsQueue& InstantMovementEffectsQueue = InputData.InputCmd.InputCollection.FindOrAddMutableDataByType<FChaosNetInstantMovementEffectsQueue>();
-	InstantMovementEffectsQueue.Effects.Empty();
-	for (const FScheduledInstantMovementEffect& ScheduledEffect : MoverComp.GetQueuedInstantMovementEffects())
+	if (!ScheduledInstantMovementEffects.IsEmpty())
 	{
-		uint8 UniqueID = NextInstantMovementEffectUniqueID++;
-		if (UniqueID == 0xFF)
+		if (bIsGeneratingInputsLocally)
+		{
+			// Input generating simulations inject instant movement effects from the Game Thread into the simulation input collection
+			// It will be used directly by the state machine (same as when it is received from the network or overwritten during resim),
+			// and will also be networked via NetInputCmd
+			InjectInstantMovementEffectsIntoInput(OutInputData);
+		}
+		else
 		{
-			UniqueID = 0;
-			NextInstantMovementEffectUniqueID = 1;
+			// Transfer queued instant movement effects to the simulation
+			// Instant effects should still be consumed by a non input generating actor on a server if it is not controlled remotely (this is an approximation of "no other instance is input producing")
+			InjectInstantMovementEffectsIntoLocalSimInput();
 		}
+		// We clear queued instant movement effects systematically to avoid accumulation, even if they ended up being discarded
+		MoverComp.ClearQueuedInstantMovementEffects();
+	}
+}
+
+void UChaosMoverBackendComponent::InjectInstantMovementEffectsIntoInput(UE::ChaosMover::FSimulationInputData& OutInputData)
+{
+	UMoverComponent& MoverComp = GetMoverComponent();
+	// Add all queued instant effects to the input so it can be networked
+	const TArray<FScheduledInstantMovementEffect>& ScheduledInstantMovementEffects = MoverComp.GetQueuedInstantMovementEffects();
+	FChaosNetInstantMovementEffectsQueue& NetInstantMovementEffectsQueue = OutInputData.InputCmd.InputCollection.FindOrAddMutableDataByType<FChaosNetInstantMovementEffectsQueue>();
+	ConvertToSerializable(ScheduledInstantMovementEffects, NetInstantMovementEffectsQueue);
+}
+
+void UChaosMoverBackendComponent::InjectInstantMovementEffectsIntoLocalSimInput()
+{
+	UMoverComponent& MoverComp = GetMoverComponent();
+	FMoverDataCollection& LocalSimInput = Simulation->GetLocalSimInput_Mutable();
+	FChaosMoverSimulationDefaultInputs& SimInputs = LocalSimInput.FindOrAddMutableDataByType<FChaosMoverSimulationDefaultInputs>();
+	ENetMode NetMode = SimInputs.OwningActor ? SimInputs.OwningActor->GetNetMode() : NM_MAX;
+	if (NetMode == NM_DedicatedServer || (NetMode == NM_ListenServer && !SimInputs.bIsRemotelyControlled))
+	{
+		const TArray<FScheduledInstantMovementEffect>& ScheduledInstantMovementEffects = MoverComp.GetQueuedInstantMovementEffects();
+		FChaosNetInstantMovementEffectsQueue& NetInstantMovementEffectsQueue = LocalSimInput.FindOrAddMutableDataByType<FChaosNetInstantMovementEffectsQueue>();
+		ConvertToSerializable(ScheduledInstantMovementEffects, NetInstantMovementEffectsQueue);
+	}
+}
 
-		InstantMovementEffectsQueue.Add(ScheduledEffect, TimeStep.ServerFrame, UniqueID);
+void UChaosMoverBackendComponent::ConvertToSerializable(const TArray<FScheduledInstantMovementEffect>& ScheduledInstantMovementEffects, FChaosNetInstantMovementEffectsQueue& OutNetInstantMovementEffectsQueue)
+{
+	OutNetInstantMovementEffectsQueue.Effects.Empty();
+	for (const FScheduledInstantMovementEffect& ScheduledEffect : ScheduledInstantMovementEffects)
+	{
+		uint8 UniqueID = GetNextUniqueSimCommandID();
 
-#if !defined(BUILD_SHIPPING) || !BUILD_SHIPPING
+		// All effects issued by the Game Thread can't roll back since the GT will not be resimmed and will not have a chance to reissue them
+		OutNetInstantMovementEffectsQueue.Add(ScheduledEffect, /*bShouldRollBack = */false, UniqueID);
+
+#if !UE_BUILD_SHIPPING
 		ENetMode NetMode = GetWorld() ? GetWorld()->GetNetMode() : NM_MAX;
-		UE_LOG(LogChaosMover, Verbose, TEXT("(%s) UChaosMoverBackendComponent::InjectInstantMovementEffects: Transferring Instant Effect Scheduled for frame %d at frame %d (Assigning Net ID %d) from the mover component to the async simulation: %s."),
-			*ToString(NetMode), ScheduledEffect.ExecutionServerFrame, TimeStep.ServerFrame, UniqueID, ScheduledEffect.Effect.IsValid() ? *ScheduledEffect.Effect->ToSimpleString() : TEXT("INVALID INSTANT EFFECT"));
-#endif // !defined(BUILD_SHIPPING) || !BUILD_SHIPPING
+		UE_LOGF(LogChaosMover, Verbose, "(%ls) UChaosMoverBackendComponent: Transferring Instant Effect Scheduled [%ls] (Assigning Net ID %d) from the mover component to the async simulation: %ls.",
+			*ToString(NetMode), *ScheduledEffect.SchedulingInfo.ToString(), UniqueID, ScheduledEffect.Effect.IsValid() ? *ScheduledEffect.Effect->ToSimpleString() : TEXT("INVALID INSTANT EFFECT"));
+#endif // !UE_BUILD_SHIPPING
 	}
 	// The effects queue is cleared in ProduceInputData
 }
 
-void UChaosMoverBackendComponent::InjectInstantMovementEffectsIntoSim(const FMoverTimeStep& TimeStep)
+void UChaosMoverBackendComponent::InjectLayeredMovesIntoSim(UE::ChaosMover::FSimulationInputData& OutInputData)
 {
+	const bool bIsGeneratingInputsLocally = !NetworkPhysicsComponent || NetworkPhysicsComponent->IsLocallyControlled();
+
 	UMoverComponent& MoverComp = GetMoverComponent();
-	FChaosMoverSimulationDefaultInputs& SimInputs = Simulation->GetLocalSimInput_Mutable().FindOrAddMutableDataByType<FChaosMoverSimulationDefaultInputs>();
+
+	const TArray<FScheduledLayeredMove>& QueuedLayeredMoves = MoverComp.GetQueuedLayeredMoves();
+	if (!QueuedLayeredMoves.IsEmpty())
+	{
+		if (bIsGeneratingInputsLocally)
+		{
+			// Input generating simulations inject moves from the Game Thread into the simulation input collection
+			// It will be used directly by the state machine (same as when it is received from the network or overwritten during resim),
+			// and will also be networked via NetInputCmd
+			InjectLayeredMovesIntoInput(OutInputData);
+		}
+		else
+		{
+			// Transfer queued moves to the simulation
+			// Moves should still be consumed by a non input generating actor on a server if it is not controlled remotely (this is an approximation of "no other instance is input producing")
+			InjectLayeredMovesIntoLocalSimInput();
+		}
+		// We clear queued moves systematically to avoid accumulation, even if they ended up being discarded
+		MoverComp.ClearQueuedLayeredMoves();
+	}
+}
+
+// Transfers queued layered moves from the mover component to the (networked) simulation input.
+// This is used in the input producing case.
+void UChaosMoverBackendComponent::InjectLayeredMovesIntoInput(UE::ChaosMover::FSimulationInputData& OutInputData)
+{
+	UMoverComponent& MoverComp = GetMoverComponent();
+	// Add all queued moves to the input so it can be networked
+	const TArray<FScheduledLayeredMove>& ScheduledLayeredMoves = MoverComp.GetQueuedLayeredMoves();
+	FChaosNetLayeredMovesQueue& NetLayeredMovesQueue = OutInputData.InputCmd.InputCollection.FindOrAddMutableDataByType<FChaosNetLayeredMovesQueue>();
+	ConvertToSerializable(ScheduledLayeredMoves, NetLayeredMovesQueue);
+}
+
+// Transfers queued layered moves from the mover component to the simulation directly, bypassing the networked input.
+// This is used in the non input producing case.
+void UChaosMoverBackendComponent::InjectLayeredMovesIntoLocalSimInput()
+{
+	UMoverComponent& MoverComp = GetMoverComponent();
+	FMoverDataCollection& LocalSimInput = Simulation->GetLocalSimInput_Mutable();
+	FChaosMoverSimulationDefaultInputs& SimInputs = LocalSimInput.FindOrAddMutableDataByType<FChaosMoverSimulationDefaultInputs>();
 	ENetMode NetMode = SimInputs.OwningActor ? SimInputs.OwningActor->GetNetMode() : NM_MAX;
 	if (NetMode == NM_DedicatedServer || (NetMode == NM_ListenServer && !SimInputs.bIsRemotelyControlled))
 	{
-		for (const FScheduledInstantMovementEffect& ScheduledEffect : MoverComp.GetQueuedInstantMovementEffects())
-		{
-			Simulation->QueueInstantMovementEffect(ScheduledEffect, /*bShouldRollBack =*/ false); // Resims do not rerun gameplay logic so we should not roll back effects issued by the game thread
+		const TArray<FScheduledLayeredMove>& ScheduledLayeredMoves = MoverComp.GetQueuedLayeredMoves();
+		FChaosNetLayeredMovesQueue& NetLayeredMovesQueue = LocalSimInput.FindOrAddMutableDataByType<FChaosNetLayeredMovesQueue>();
+		ConvertToSerializable(ScheduledLayeredMoves, NetLayeredMovesQueue);
+	}
+}
 
-			UE_LOG(LogChaosMover, Verbose, TEXT("(%s) ProduceInputData: Async simulation received an Instant Movement Effect from the Mover Component, scheduled for frame %d at frame %d: %s"),
-				*ToString(NetMode), ScheduledEffect.ExecutionServerFrame, TimeStep.ServerFrame, ScheduledEffect.Effect.IsValid() ? *ScheduledEffect.Effect->ToSimpleString() : TEXT("INVALID INSTANT EFFECT"));
-		}
+uint8 UChaosMoverBackendComponent::GetNextUniqueSimCommandID()
+{
+	uint8 UniqueID = NextUniqueSimCommandID++;
+	if (UniqueID == 0xFF)
+	{
+		UniqueID = 0;
+		NextUniqueSimCommandID = 1;
 	}
-	// The effects queue is cleared in ProduceInputData
+	return UniqueID;
+}
+
+void UChaosMoverBackendComponent::ConvertToSerializable(const TArray<FScheduledLayeredMove>& ScheduledLayeredMoves, FChaosNetLayeredMovesQueue& OutNetLayeredMovesQueue)
+{
+	OutNetLayeredMovesQueue.Moves.Empty();
+	for (const FScheduledLayeredMove& ScheduledMove : ScheduledLayeredMoves)
+	{
+		uint8 UniqueID = GetNextUniqueSimCommandID();
+
+		// All moves issued by the Game Thread can't roll back since the GT will not be resimmed and will not have a chance to reissue them
+		OutNetLayeredMovesQueue.Add(ScheduledMove, /*bShouldRollBack = */false, UniqueID);
+
+#if !UE_BUILD_SHIPPING
+		ENetMode NetMode = GetWorld() ? GetWorld()->GetNetMode() : NM_MAX;
+		UE_LOGF(LogChaosMover, Verbose, "(%ls) UChaosMoverBackendComponent: Transferring move scheduled [%ls] (Assigning Net ID %d) from the mover component to the async simulation: %ls.",
+			*ToString(NetMode), *ScheduledMove.SchedulingInfo.ToString(), UniqueID, ScheduledMove.Move.IsValid() ? *ScheduledMove.Move->ToSimpleString() : TEXT("INVALID MOVE"));
+#endif // !UE_BUILD_SHIPPING
+	}
+	// The moves queue is cleared in ProduceInputData
 }
 
 void UChaosMoverBackendComponent::ConsumeOutputData(const FMoverTimeStep& TimeStep, const UE::ChaosMover::FSimulationOutputData& OutputData)
diff --git a/Engine/Plugins/Experimental/ChaosMover/Source/ChaosMover/Private/ChaosMoverConsoleVariables.cpp b/Engine/Plugins/Experimental/ChaosMover/Source/ChaosMover/Private/ChaosMoverConsoleVariables.cpp
--- a/Engine/Plugins/Experimental/ChaosMover/Source/ChaosMover/Private/ChaosMoverConsoleVariables.cpp
+++ b/Engine/Plugins/Experimental/ChaosMover/Source/ChaosMover/Private/ChaosMoverConsoleVariables.cpp
@@ -41,7 +41,11 @@ namespace UE::ChaosMover
 
 		int32 InstantMovementEffectIDHistorySize = 30;
 		FAutoConsoleVariableRef CVarChaosMoverInstantMovementEffectIDHistorySize(TEXT("ChaosMover.Networking.InstantMovementEffectIDHistorySize"),
-			InstantMovementEffectIDHistorySize, TEXT("Number of past frames for which we keep the last seen instant movement effect IDs, to avoid duplicate processing of events when input is repeated\n"));
+			InstantMovementEffectIDHistorySize, TEXT("Number of past frames for which we keep the last seen instant movement effect IDs, to avoid duplicate processing when input is repeated\n"));
+
+		int32 LayeredMoveIDHistorySize = 30;
+		FAutoConsoleVariableRef CVarChaosMoverLayeredMoveIDHistorySize(TEXT("ChaosMover.Networking.LayeredMoveIDHistorySize"),
+			LayeredMoveIDHistorySize, TEXT("Number of past frames for which we keep the last seen layered move IDs, to avoid duplicate processing when input is repeated\n"));
 
 		int32 FrameDifferenceLeniencyForEventTimeComparison = 1;
 		FAutoConsoleVariableRef CVarChaosMoverFrameDifferenceLeniencyForEventTimeComparison(TEXT("ChaosMover.Networking.FrameDifferenceLeniencyForEventTimeComparison"),
diff --git a/Engine/Plugins/Experimental/ChaosMover/Source/ChaosMover/Private/ChaosMoverSimulation.cpp b/Engine/Plugins/Experimental/ChaosMover/Source/ChaosMover/Private/ChaosMoverSimulation.cpp
--- a/Engine/Plugins/Experimental/ChaosMover/Source/ChaosMover/Private/ChaosMoverSimulation.cpp
+++ b/Engine/Plugins/Experimental/ChaosMover/Source/ChaosMover/Private/ChaosMoverSimulation.cpp
@@ -1243,12 +1243,67 @@ void UChaosMoverSimulation::QueueInstantMovementEffect(TSharedPtr<FInstantMoveme
 {
 	// We always use a server frame to schedule, never a time. This is because ChaosMover is always in async mode
     // (if we're not, UChaosMoverBackendComponent::InitializeComponent should have warned us to fix that)
-	QueueInstantMovementEffect(FScheduledInstantMovementEffect{ InternalServerFrame, /* ExecutionDateSeconds = */ 0.0, /* bIsAsyncMode = */ true, InstantMovementEffect }, bShouldRollBack);
+	QueueInstantMovementEffect(FScheduledInstantMovementEffect{
+									FMoverSchedulingInfo{
+										/* ServerIssuanceTime  = */ FMoverTime{InternalServerFrame, 0.0},
+										/* ServerExecutionTime = */ FMoverTime{InternalServerFrame, 0.0},
+										/* bIsFixedDt          = */ true},
+									InstantMovementEffect },
+								bShouldRollBack);
 }
 
 void UChaosMoverSimulation::QueueInstantMovementEffect(const FScheduledInstantMovementEffect& ScheduledInstantMovementEffect, bool bShouldRollBack)
 {
-	StateMachine.QueueInstantMovementEffect(FChaosScheduledInstantMovementEffect{ InternalServerFrame, bShouldRollBack, ScheduledInstantMovementEffect });
+	StateMachine.QueueInstantMovementEffect(FChaosScheduledInstantMovementEffect{ bShouldRollBack, ScheduledInstantMovementEffect });
+}
+
+void UChaosMoverSimulation::K2_QueueLayeredMove(const int32& MoveAsRawData)
+{
+	// This will never be called, the exec version below will be hit instead
+	checkNoEntry();
+}
+
+DEFINE_FUNCTION(UChaosMoverSimulation::execK2_QueueLayeredMove)
+{
+	Stack.StepCompiledIn<FStructProperty>(nullptr);
+	void* MovePtr = Stack.MostRecentPropertyAddress;
+	FStructProperty* StructProp = CastField<FStructProperty>(Stack.MostRecentProperty);
+
+	P_FINISH;
+
+	P_NATIVE_BEGIN;
+
+	const bool bHasValidStructProp = StructProp && StructProp->Struct && StructProp->Struct->IsChildOf(FLayeredMoveBase::StaticStruct());
+
+	if (ensureMsgf((bHasValidStructProp && MovePtr), TEXT("An invalid type (%s) was sent to a QueueLayeredMove node. A struct derived from FLayeredMoveBase is required. No Move will be queued."),
+		StructProp ? *GetNameSafe(StructProp->Struct) : *Stack.MostRecentProperty->GetClass()->GetName()))
+	{
+		// Could we steal this instead of cloning? (move semantics)
+		FLayeredMoveBase* MoveAsBasePtr = reinterpret_cast<FLayeredMoveBase*>(MovePtr);
+		FLayeredMoveBase* ClonedMove = MoveAsBasePtr->Clone();
+
+		P_THIS->QueueLayeredMove(TSharedPtr<FLayeredMoveBase>(ClonedMove), /*bShouldRollBack =*/ true); // Physics Thread moves can be rolled back because we expect resim to rerun all Physics Thread logic
+	}
+
+	P_NATIVE_END;
+}
+
+void UChaosMoverSimulation::QueueLayeredMove(TSharedPtr<FLayeredMoveBase> LayeredMove, bool bShouldRollBack)
+{
+	// We always use a server frame to schedule, never a time. This is because ChaosMover is always in async mode
+	// (if we're not, UChaosMoverBackendComponent::InitializeComponent should have warned us to fix that)
+	QueueLayeredMove(FScheduledLayeredMove{
+						FMoverSchedulingInfo{
+							/* ServerIssuanceTime  = */ FMoverTime{InternalServerFrame, 0.0},
+							/* ServerExecutionTime = */ FMoverTime{InternalServerFrame, 0.0},
+							/* bIsFixedDt          = */ true},
+						LayeredMove },
+					 bShouldRollBack);
+}
+
+void UChaosMoverSimulation::QueueLayeredMove(const FScheduledLayeredMove& ScheduledLayeredMove, bool bShouldRollBack)
+{
+	StateMachine.QueueLayeredMove(FChaosScheduledLayeredMove{ bShouldRollBack, ScheduledLayeredMove });
 }
 
 DEFINE_FUNCTION(UChaosMoverSimulation::execK2_QueueMovementModifier)
@@ -1353,30 +1408,30 @@ void UChaosMoverSimulation::AttemptTeleport(const FMoverTimeStep& TimeStep, cons
 
 		AddEvent(MakeShared<FTeleportSucceededEventData>(TimeStep, FromTransform.GetLocation(), FromTransform.GetRotation(), ToLocation, ToRotation));
 
-#if !defined(BUILD_SHIPPING) || !BUILD_SHIPPING
+#if !UE_BUILD_SHIPPING
 		const FChaosMoverSimulationDefaultInputs* SimInputs = LocalSimInput.FindDataByType<FChaosMoverSimulationDefaultInputs>();
 		if (SimInputs && SimInputs->OwningActor)
 		{
-			UE_LOG(LogChaosMover, Log, TEXT("%s: Actor %s teleported to %s, %s, teleport frame = %d)"),
-			*ToString(SimInputs->OwningActor->GetNetMode()), *GetNameSafe(SimInputs->OwningActor),
-			*ToLocation.ToString(), *FRotator(ToRotation).ToString(), TimeStep.ServerFrame);
+			UE_LOGF(LogChaosMover, Log, "%ls: Actor %ls teleported to %ls, %ls, teleport frame = %d)",
+				*ToString(SimInputs->OwningActor->GetNetMode()), *GetNameSafe(SimInputs->OwningActor),
+				*ToLocation.ToString(), *FRotator(ToRotation).ToString(), TimeStep.ServerFrame);
 		}
-#endif // !defined(BUILD_SHIPPING) || !BUILD_SHIPPING
+#endif // !UE_BUILD_SHIPPING
 	}
 	else
 	{
 		ETeleportFailureReason FailureReason = ETeleportFailureReason::Reason_NotAvailable;
 		AddEvent(MakeShared<FTeleportFailedEventData>(TimeStep, FromTransform.GetLocation(), FromTransform.GetRotation(), ToLocation, ToRotation, FailureReason));
 
-#if !defined(BUILD_SHIPPING) || !BUILD_SHIPPING
+#if !UE_BUILD_SHIPPING
 		const FChaosMoverSimulationDefaultInputs* SimInputs = LocalSimInput.FindDataByType<FChaosMoverSimulationDefaultInputs>();
 		if (SimInputs && SimInputs->OwningActor)
 		{
-			UE_LOG(LogChaosMover, Log, TEXT("%s: Actor %s could NOT teleport to %s, %s, teleport frame = %d. Reason = %s"),
-			*GetNameSafe(SimInputs->OwningActor), *ToString(SimInputs->OwningActor->GetNetMode()),
-			*ToLocation.ToString(), *FRotator(ToRotation).ToString(), TimeStep.ServerFrame,
-			*StaticEnum<ETeleportFailureReason>()->GetNameStringByValue(static_cast<int64>(FailureReason)));
+			UE_LOGF(LogChaosMover, Log, "%ls: Actor %ls could NOT teleport to %ls, %ls, teleport frame = %d. Reason = %ls",
+				*ToString(SimInputs->OwningActor->GetNetMode()), *GetNameSafe(SimInputs->OwningActor),
+				*ToLocation.ToString(), *FRotator(ToRotation).ToString(), TimeStep.ServerFrame,
+				*StaticEnum<ETeleportFailureReason>()->GetNameStringByValue(static_cast<int64>(FailureReason)));
 		}
-#endif // !defined(BUILD_SHIPPING) || !BUILD_SHIPPING
+#endif // !UE_BUILD_SHIPPING
 	}
 }
diff --git a/Engine/Plugins/Experimental/ChaosMover/Source/ChaosMover/Private/ChaosMoverSimulationTypes.cpp b/Engine/Plugins/Experimental/ChaosMover/Source/ChaosMover/Private/ChaosMoverSimulationTypes.cpp
--- a/Engine/Plugins/Experimental/ChaosMover/Source/ChaosMover/Private/ChaosMoverSimulationTypes.cpp
+++ b/Engine/Plugins/Experimental/ChaosMover/Source/ChaosMover/Private/ChaosMoverSimulationTypes.cpp
@@ -186,15 +186,16 @@ void FChaosWaterResultData::Decay(float DecayAmount)
 {
 }
 
-void FChaosNetInstantMovementEffectsQueue::Add(const FScheduledInstantMovementEffect& InScheduledEffect, int32 IssuanceServerFrame, uint8 UniqueID)
+void FChaosNetInstantMovementEffectsQueue::Add(const FScheduledInstantMovementEffect& InScheduledEffect, bool bShouldRollBack, uint8 UniqueID)
 {
 	TCheckedObjPtr<FInstantMovementEffect> InstantMovementEffect = InScheduledEffect.Effect.Get();
-	if (ensure(InstantMovementEffect.IsValid()))
+	if (ensure(InstantMovementEffect.IsValid() && InstantMovementEffect->GetScriptStruct()))
 	{
 		FChaosNetInstantMovementEffect& AddedInstancedEffect = Effects.AddDefaulted_GetRef();
-		AddedInstancedEffect.ExecutionServerFrame = InScheduledEffect.ExecutionServerFrame;
-		AddedInstancedEffect.IssuanceServerFrame = IssuanceServerFrame;	
+		AddedInstancedEffect.ExecutionServerFrame = InScheduledEffect.SchedulingInfo.ServerExecutionTime.FrameCount;
+		AddedInstancedEffect.IssuanceServerFrame = InScheduledEffect.SchedulingInfo.ServerIssuanceTime.FrameCount;
 		AddedInstancedEffect.UniqueID = UniqueID;
+		AddedInstancedEffect.bShouldRollBack = bShouldRollBack;
 		AddedInstancedEffect.Effect.InitializeAsScriptStruct(InstantMovementEffect->GetScriptStruct(), (const uint8*)InScheduledEffect.Effect.Get());
 	}
 }
@@ -204,7 +205,7 @@ FChaosNetInstantMovementEffectsQueue& FChaosNetInstantMovementEffectsQueue::oper
 	Effects.Empty(/*Slack =*/ ScheduledInstantEffects.Num());
 	for (const FChaosScheduledInstantMovementEffect& ScheduledInstantEffect : ScheduledInstantEffects)
 	{
-		Add(ScheduledInstantEffect.ScheduledEffect, ScheduledInstantEffect.IssuanceServerFrame, 0xFF/* This is used for debugging, unique ID doesn't apply */);
+		Add(ScheduledInstantEffect.ScheduledEffect, ScheduledInstantEffect.bShouldRollBack, 0xFF/* This is used for debugging, unique ID doesn't apply */);
 	}
 	return *this;
 }
@@ -216,7 +217,8 @@ bool FChaosNetInstantMovementEffectsQueue::NetSerialize(FArchive& Ar, UPackageMa
 	uint8 NumEffectsToSerialize;
 	if (Ar.IsSaving())
 	{
-		NumEffectsToSerialize = Effects.Num();
+		ensureMsgf(Effects.Num() <= 255, TEXT("More than 255 effects to serialize, some effects will be dropped"));
+		NumEffectsToSerialize = FMath::Min(Effects.Num(), 255);
 	}
 	Ar << NumEffectsToSerialize;
 
@@ -261,10 +263,25 @@ bool FChaosNetInstantMovementEffectsQueue::NetSerialize(FArchive& Ar, UPackageMa
 				if (Ar.IsLoading())
 				{
 					ScheduledEffect.Effect.InitializeAsScriptStruct(ScriptStruct.Get());
+					ScheduledEffect.bShouldRollBack = true; // Networked effects should always roll back, only game thread issued ones don't
+					// In the loading case it's only now that Effect can be set to a valid initialized struct
+					if (ensure(ScheduledEffect.Effect.IsValid()))
+					{
+						Effect = &ScheduledEffect.Effect.GetMutable();
+					}
 				}
 
-				FInstantMovementEffect& InstantMovementEffect = ScheduledEffect.Effect.GetMutable();
-				InstantMovementEffect.NetSerialize(Ar);
+				if (ensure(Effect))
+				{
+					Effect->NetSerialize(Ar);
+				}
+				else
+				{
+					UE_LOG(LogMover, Error, TEXT("FInstantMovementEffectsQueue::NetSerialize: Failed to serialize effect"));
+					Ar.SetError();
+					bOutSuccess = false;
+					return false;
+				}
 			}
 			else
 			{
@@ -298,7 +315,10 @@ void FChaosNetInstantMovementEffectsQueue::ToString(FAnsiStringBuilderBase& Out)
 	for (int32 EffectIndex = 0; EffectIndex < Effects.Num(); ++EffectIndex)
 	{
 		const FChaosNetInstantMovementEffect& ScheduledEffect = Effects[EffectIndex];
-		Out.Appendf("Effect Index %d: ServerFrame %d, %s", EffectIndex, ScheduledEffect.ExecutionServerFrame, ScheduledEffect.Effect.IsValid() ? *ScheduledEffect.Effect.Get().ToSimpleString() : TEXT("INVALID INSTANCED STRUCT EFFECT") );
+		Out.Appendf("Effect Index %d: ServerFrame %d, %s",
+			EffectIndex,
+			ScheduledEffect.ExecutionServerFrame,
+			ScheduledEffect.Effect.IsValid() ? TCHAR_TO_ANSI(*ScheduledEffect.Effect.Get().ToSimpleString()) : "INVALID INSTANCED STRUCT EFFECT");
 		Out.AppendChar('\n');
 	}
 	Out.Appendf("--------------------------------------------------------------------------------\n");
@@ -374,6 +394,216 @@ void FChaosNetInstantMovementEffectsQueue::Decay(float DecayAmount)
 
 }
 
+void FChaosNetLayeredMovesQueue::Add(const FScheduledLayeredMove& InScheduledLayeredMove, bool bShouldRollBack, uint8 UniqueID)
+{
+	TCheckedObjPtr<FLayeredMoveBase> LayeredMove = InScheduledLayeredMove.Move.Get();
+	if (ensure(LayeredMove.IsValid() && LayeredMove->GetScriptStruct()))
+	{
+		FChaosNetLayeredMove& AddedInstancedEffect = Moves.AddDefaulted_GetRef();
+		AddedInstancedEffect.ExecutionServerFrame = InScheduledLayeredMove.SchedulingInfo.ServerExecutionTime.FrameCount;
+		AddedInstancedEffect.IssuanceServerFrame = InScheduledLayeredMove.SchedulingInfo.ServerIssuanceTime.FrameCount;
+		AddedInstancedEffect.UniqueID = UniqueID;
+		AddedInstancedEffect.bShouldRollBack = bShouldRollBack;
+		AddedInstancedEffect.Move.InitializeAsScriptStruct(LayeredMove->GetScriptStruct(), (const uint8*)InScheduledLayeredMove.Move.Get());
+	}
+}
+
+// This function is used for debugging (sending to CVD)
+FChaosNetLayeredMovesQueue& FChaosNetLayeredMovesQueue::operator=(const TArray<FChaosScheduledLayeredMove>& ScheduledMoves)
+{
+	Moves.Empty(/*Slack =*/ ScheduledMoves.Num());
+	for (const FChaosScheduledLayeredMove& ScheduledMove : ScheduledMoves)
+	{
+		Add(ScheduledMove.ScheduledLayeredMove, ScheduledMove.bShouldRollBack, 0xFF/* This is used for debugging, unique ID doesn't apply */);
+	}
+	return *this;
+}
+
+FMoverDataStructBase* FChaosNetLayeredMovesQueue::Clone() const
+{
+	FChaosNetLayeredMovesQueue* CopyPtr = new FChaosNetLayeredMovesQueue(*this);
+	return CopyPtr;
+}
+
+bool FChaosNetLayeredMovesQueue::NetSerialize(FArchive& Ar, UPackageMap* Map, bool& bOutSuccess)
+{
+	bOutSuccess = true;
+
+	uint8 NumMovesToSerialize;
+	if (Ar.IsSaving())
+	{
+		ensureMsgf(Moves.Num() <= 255, TEXT("More than 255 moves to serialize, some moves will be dropped"));
+		NumMovesToSerialize = FMath::Min(Moves.Num(), 255);
+	}
+	Ar << NumMovesToSerialize;
+
+	if (Ar.IsLoading())
+	{
+		Moves.SetNumZeroed(NumMovesToSerialize);
+	}
+
+	for (int32 MoveIndex = 0; MoveIndex < NumMovesToSerialize && !Ar.IsError(); ++MoveIndex)
+	{
+		FChaosNetLayeredMove& ScheduledMove = Moves[MoveIndex];
+		FLayeredMoveBase* Move = ScheduledMove.Move.IsValid() ? &ScheduledMove.Move.GetMutable() : nullptr;
+		TCheckedObjPtr<UScriptStruct> ScriptStruct = Ar.IsLoading() ? FLayeredMoveBase::StaticStruct() : (Move ? Move->GetScriptStruct() : nullptr);
+
+		Ar << ScriptStruct;
+
+		if (ScriptStruct.IsValid())
+		{
+			// Restrict replication to derived classes of FLayeredMoveBase for security reasons:
+			// If FLayeredMovesQueue is replicated through a Server RPC, we need to prevent clients from sending us
+			// arbitrary ScriptStructs due to the allocation/reliance on GetCppStructOps below which could trigger a server crash
+			// for invalid structs. All provided sources are direct children of FLayeredMoveBase and we never expect to have deep hierarchies
+			// so this should not be too costly
+			bool bIsDerivedFromBase = false;
+			UStruct* CurrentSuperStruct = ScriptStruct->GetSuperStruct();
+			while (CurrentSuperStruct)
+			{
+				if (CurrentSuperStruct == FLayeredMoveBase::StaticStruct())
+				{
+					bIsDerivedFromBase = true;
+					break;
+				}
+				CurrentSuperStruct = CurrentSuperStruct->GetSuperStruct();
+			}
+
+			if (bIsDerivedFromBase)
+			{
+				Ar << ScheduledMove.ExecutionServerFrame;
+				Ar << ScheduledMove.IssuanceServerFrame;
+				Ar << ScheduledMove.UniqueID;
+
+				if (Ar.IsLoading())
+				{
+					ScheduledMove.Move.InitializeAsScriptStruct(ScriptStruct.Get());
+					ScheduledMove.bShouldRollBack = true; // Networked moves should always roll back, only game thread issued ones don't
+					// In the loading case it's only now that Move can be set to a valid initialized struct
+					if (ensure(ScheduledMove.Move.IsValid()))
+					{
+						Move = &ScheduledMove.Move.GetMutable();
+					}
+				}
+
+				if (ensure(Move))
+				{
+					Move->NetSerialize(Ar);
+				}
+				else
+				{
+					UE_LOG(LogMover, Error, TEXT("FChaosNetLayeredMovesQueue::NetSerialize: Failed to serialize move"));
+					Ar.SetError();
+					bOutSuccess = false;
+					return false;
+				}
+			}
+			else
+			{
+				UE_LOG(LogMover, Error, TEXT("FChaosNetLayeredMovesQueue::NetSerialize: ScriptStruct not derived from FLayeredMoveBase attempted to serialize."));
+				Ar.SetError();
+				bOutSuccess = false;
+				return false;
+			}
+		}
+		else if (ScriptStruct.IsError())
+		{
+			UE_LOG(LogMover, Error, TEXT("FChaosNetLayeredMovesQueue::NetSerialize: Invalid ScriptStruct serialized."));
+			Ar.SetError();
+			bOutSuccess = false;
+			return false;
+		}
+	}
+
+	return true;
+}
+
+UScriptStruct* FChaosNetLayeredMovesQueue::GetScriptStruct() const
+{
+	return StaticStruct();
+}
+
+void FChaosNetLayeredMovesQueue::ToString(FAnsiStringBuilderBase& Out) const
+{
+	Out.Appendf("Layered Moves Queue -------------------------------------------------\n");
+	for (int32 MoveIndex = 0; MoveIndex < Moves.Num(); ++MoveIndex)
+	{
+		const FChaosNetLayeredMove& ScheduledMove = Moves[MoveIndex];
+		Out.Appendf("Move Index %d: ServerFrame %d, %s",
+			MoveIndex,
+			ScheduledMove.ExecutionServerFrame,
+			ScheduledMove.Move.IsValid() ? TCHAR_TO_ANSI(*ScheduledMove.Move.Get().ToSimpleString()) : "INVALID INSTANCED STRUCT MOVE");
+		Out.AppendChar('\n');
+	}
+	Out.Appendf("--------------------------------------------------------------------------------\n");
+}
+
+bool FChaosNetLayeredMovesQueue::ShouldReconcile(const FMoverDataStructBase& AuthorityState) const
+{
+	const FChaosNetLayeredMovesQueue& TypedAuthority = static_cast<const FChaosNetLayeredMovesQueue&>(AuthorityState);
+	if (Moves.Num() != TypedAuthority.Moves.Num())
+	{
+		return true;
+	}
+
+	for (int MoveIndex = 0; MoveIndex < Moves.Num(); ++MoveIndex)
+	{
+		const FChaosNetLayeredMove& ScheduledMove = Moves[MoveIndex];
+		const FChaosNetLayeredMove& AuthorityScheduledMove = TypedAuthority.Moves[MoveIndex];
+
+		if (ScheduledMove.ExecutionServerFrame != AuthorityScheduledMove.ExecutionServerFrame)
+		{
+			return true;
+		}
+
+		if (ScheduledMove.Move.IsValid() != AuthorityScheduledMove.Move.IsValid())
+		{
+			return true;
+		}
+
+		if (ScheduledMove.Move.IsValid())
+		{
+			if (ScheduledMove.Move.GetScriptStruct() != AuthorityScheduledMove.Move.GetScriptStruct())
+			{
+				return true;
+			}
+
+			// This allows us to skip implementing operator== for all moves
+			FArchiveCrc32 Crc1;
+			(const_cast<FChaosNetLayeredMove&>(ScheduledMove)).Move.GetMutable().NetSerialize(Crc1);
+			FArchiveCrc32 Crc2;
+			(const_cast<FChaosNetLayeredMove&>(AuthorityScheduledMove)).Move.GetMutable().NetSerialize(Crc2);
+			bool bAreEqual = (Crc1.GetCrc() == Crc2.GetCrc());
+			if (!bAreEqual)
+			{
+				return true;
+			}
+		}
+	}
+
+	return false;
+}
+
+void FChaosNetLayeredMovesQueue::Interpolate(const FMoverDataStructBase& From, const FMoverDataStructBase& To, float Pct)
+{
+	const FChaosNetLayeredMovesQueue& TypedFrom = static_cast<const FChaosNetLayeredMovesQueue&>(From);
+	const FChaosNetLayeredMovesQueue& TypedTo = static_cast<const FChaosNetLayeredMovesQueue&>(To);
+
+	Moves = (Pct < 0.5f) ? TypedFrom.Moves : TypedTo.Moves;
+
+}
+
+void FChaosNetLayeredMovesQueue::Merge(const FMoverDataStructBase& From)
+{
+	const FChaosNetLayeredMovesQueue& TypedFrom = static_cast<const FChaosNetLayeredMovesQueue&>(From);
+	Moves.Append(TypedFrom.Moves);
+}
+
+void FChaosNetLayeredMovesQueue::Decay(float DecayAmount)
+{
+
+}
+
 bool FAsyncLocalOnlyInstantMovementEffect::ApplyMovementEffect_Async(FApplyMovementEffectParams_Async& ApplyEffectParams, FMoverSyncState& OutputState)
 {
 	// AsyncFunction will not be set if this is not executing on the instance that queued this effect, nothing will execute
diff --git a/Engine/Plugins/Experimental/ChaosMover/Source/ChaosMover/Private/ChaosMoverStateMachine.cpp b/Engine/Plugins/Experimental/ChaosMover/Source/ChaosMover/Private/ChaosMoverStateMachine.cpp
--- a/Engine/Plugins/Experimental/ChaosMover/Source/ChaosMover/Private/ChaosMoverStateMachine.cpp
+++ b/Engine/Plugins/Experimental/ChaosMover/Source/ChaosMover/Private/ChaosMoverStateMachine.cpp
@@ -21,6 +21,15 @@
 #if WITH_CHAOS_VISUAL_DEBUGGER
 #include "ChaosVisualDebugger/ChaosVisualDebuggerTrace.h"
 #endif
+
+#if !UE_BUILD_SHIPPING
+#define NETMODE_TO_STR(SimInputs) (SimInputs && IsValid(SimInputs->OwningActor) ? *ToString(SimInputs->OwningActor->GetNetMode()) : TEXT("UNKNOWN NET MODE"))
+#define MOVESHAREDPTR_TO_STR(Move) Move.IsValid() ? *Move->ToSimpleString() : TEXT("INVALID MOVE")
+#define EFFECTSHAREDPTR_TO_STR(Effect) Effect.IsValid() ? *Effect->ToSimpleString() : TEXT("INVALID EFFECT")
+#define MOVEINSTANCEDSTRUCT_TO_STR(Move) Move.IsValid() ? *Move.Get().ToSimpleString() : TEXT("INVALID MOVE")
+#define EFFECTINSTANCEDSTRUCT_TO_STR(Effect) Effect.IsValid() ? *Effect.Get().ToSimpleString() : TEXT("INVALID EFFECT")
+#endif
+
 namespace UE::ChaosMover
 {
 	FMoverStateMachine::FMoverStateMachine()
@@ -266,7 +275,7 @@ namespace UE::ChaosMover
 			for (int EffectIndex = 0; EffectIndex < QueuedInstantEffects.Num(); EffectIndex++)
 			{
 				const FChaosScheduledInstantMovementEffect& QueuedEffect = QueuedInstantEffects[EffectIndex];
-				if (QueuedEffect.ScheduledEffect.ShouldExecuteAtFrame(TimeStep.ServerFrame))
+				if (QueuedEffect.ScheduledEffect.SchedulingInfo.ShouldExecuteAtFrame(TimeStep.ServerFrame))
 				{
 					return true;
 				}
@@ -289,29 +298,89 @@ namespace UE::ChaosMover
 			{
 				if (!InstantMovementEffectsIDHistory.WasIDAlreadySeen(NetInstantMovementEffect.UniqueID))
 				{
-					QueueInstantMovementEffect(FChaosScheduledInstantMovementEffect{ NetInstantMovementEffect.IssuanceServerFrame, NetInstantMovementEffect.bShouldRollBack, NetInstantMovementEffect.AsScheduledInstantMovementEffect() });
+					QueueInstantMovementEffect(FChaosScheduledInstantMovementEffect{
+												NetInstantMovementEffect.bShouldRollBack,
+												NetInstantMovementEffect.AsScheduledInstantMovementEffect() });
 
 					InstantMovementEffectsIDHistory.AddID(TimeStep.ServerFrame, NetInstantMovementEffect.UniqueID);
 
-#if !defined(BUILD_SHIPPING) || !BUILD_SHIPPING
+#if !UE_BUILD_SHIPPING
 					if (const FChaosMoverSimulationDefaultInputs* SimInputs = Simulation->GetLocalSimInput().FindDataByType<FChaosMoverSimulationDefaultInputs>())
 					{
 						UE_LOG(LogChaosMover, Verbose, TEXT("(%s, %s, Current Frame = %d) FMoverStateMachine::ReceiveInstantMovementEffects: Received Instant Effect (ID=%d), issued at frame %d scheduled for frame %d: %s"),
-							*ToString(SimInputs->OwningActor->GetNetMode()), SimInputs->bIsGeneratingInputsLocally ? TEXT("Input Source") : TEXT("Input Consumer"), TimeStep.ServerFrame,
-							NetInstantMovementEffect.UniqueID, NetInstantMovementEffect.IssuanceServerFrame, NetInstantMovementEffect.ExecutionServerFrame, NetInstantMovementEffect.Effect.IsValid() ? *NetInstantMovementEffect.Effect.Get().ToSimpleString() : TEXT("INVALID INSTANT EFFECT"));
+							NETMODE_TO_STR(SimInputs),
+							SimInputs->bIsGeneratingInputsLocally ? TEXT("Input Source") : TEXT("Input Consumer"),
+							TimeStep.ServerFrame,
+							NetInstantMovementEffect.UniqueID,
+							NetInstantMovementEffect.IssuanceServerFrame,
+							NetInstantMovementEffect.ExecutionServerFrame,
+							EFFECTINSTANCEDSTRUCT_TO_STR(NetInstantMovementEffect.Effect));
 					}
-#endif // !defined(BUILD_SHIPPING) || !BUILD_SHIPPING
+#endif // !UE_BUILD_SHIPPING
 				}
 				else
 				{
-#if !defined(BUILD_SHIPPING) || !BUILD_SHIPPING
+#if !UE_BUILD_SHIPPING
 					if (const FChaosMoverSimulationDefaultInputs* SimInputs = Simulation->GetLocalSimInput().FindDataByType<FChaosMoverSimulationDefaultInputs>())
 					{
 						UE_LOG(LogChaosMover, Verbose, TEXT("(%s, %s, Current Frame = %d) FMoverStateMachine::ReceiveInstantMovementEffects: Skipping Instant Effect (ID=%d), ID already seen!"),
-							*ToString(SimInputs->OwningActor->GetNetMode()), SimInputs->bIsGeneratingInputsLocally ? TEXT("Input Source") : TEXT("Input Consumer"), TimeStep.ServerFrame,
+							NETMODE_TO_STR(SimInputs),
+							SimInputs->bIsGeneratingInputsLocally ? TEXT("Input Source") : TEXT("Input Consumer"),
+							TimeStep.ServerFrame,
 							NetInstantMovementEffect.UniqueID);
 					}
-#endif // !defined(BUILD_SHIPPING) || !BUILD_SHIPPING
+#endif // !UE_BUILD_SHIPPING
+				}
+			}
+		}
+	}
+
+	void FMoverStateMachine::ReceiveLayeredMoves(const FMoverTimeStep& TimeStep, const FChaosNetLayeredMovesQueue* LayeredMovesQueue)
+	{
+		QUICK_SCOPE_CYCLE_COUNTER(MoverStateMachine_ReceiveLayeredMoves);
+
+		LayeredMovesIDHistory.CullOldFrames(TimeStep.ServerFrame, CVars::LayeredMoveIDHistorySize);
+
+		// Copy queued layered moves from the input to the state machine
+		// After this, the state machine works on its own queue, into which it can enqueue moves while stepping
+		for (const FChaosNetLayeredMove& NetLayeredMove : LayeredMovesQueue->Moves)
+		{
+			if (ensure(NetLayeredMove.Move.IsValid()))
+			{
+				if (!LayeredMovesIDHistory.WasIDAlreadySeen(NetLayeredMove.UniqueID))
+				{
+					QueueLayeredMove(FChaosScheduledLayeredMove{
+						NetLayeredMove.bShouldRollBack,
+						NetLayeredMove.AsScheduledLayeredMove() });
+
+					LayeredMovesIDHistory.AddID(TimeStep.ServerFrame, NetLayeredMove.UniqueID);
+
+#if !UE_BUILD_SHIPPING
+					if (const FChaosMoverSimulationDefaultInputs* SimInputs = Simulation->GetLocalSimInput().FindDataByType<FChaosMoverSimulationDefaultInputs>())
+					{
+						UE_LOG(LogChaosMover, Verbose, TEXT("(%s, %s, Current Frame = %d) FMoverStateMachine::ReceiveLayeredMoves: Received Layered Move (ID=%d), issued at frame %d scheduled for frame %d: %s"),
+							NETMODE_TO_STR(SimInputs),
+							SimInputs->bIsGeneratingInputsLocally ? TEXT("Input Source") : TEXT("Input Consumer"),
+							TimeStep.ServerFrame,
+							NetLayeredMove.UniqueID,
+							NetLayeredMove.IssuanceServerFrame,
+							NetLayeredMove.ExecutionServerFrame,
+							MOVEINSTANCEDSTRUCT_TO_STR(NetLayeredMove.Move));
+					}
+#endif // !UE_BUILD_SHIPPING
+				}
+				else
+				{
+#if !UE_BUILD_SHIPPING
+					if (const FChaosMoverSimulationDefaultInputs* SimInputs = Simulation->GetLocalSimInput().FindDataByType<FChaosMoverSimulationDefaultInputs>())
+					{
+						UE_LOG(LogChaosMover, Verbose, TEXT("(%s, %s, Current Frame = %d) FMoverStateMachine::ReceiveLayeredMoves: Skipping Layered Move (ID=%d), ID already seen!"),
+							NETMODE_TO_STR(SimInputs),
+							SimInputs->bIsGeneratingInputsLocally ? TEXT("Input Source") : TEXT("Input Consumer"),
+							TimeStep.ServerFrame,
+							NetLayeredMove.UniqueID);
+					}
+#endif // !UE_BUILD_SHIPPING
 				}
 			}
 		}
@@ -332,15 +401,34 @@ namespace UE::ChaosMover
 
 		InternalSimTime = TimeStep;
 
+		// Receive all sim commands (instant movement effects, queued layered moves, etc). Some may be scheduled to happen in the future.
+
 		// Copy queued instant movement effects from the input to the state machine
 		// After this, the state machine works on its own queue, into which it can enqueue instant movement effects while stepping
+		// Most will come from the networked input
 		if (const FChaosNetInstantMovementEffectsQueue* InstantMovementEffectsQueue = StartState.InputCmd.InputCollection.FindDataByType<FChaosNetInstantMovementEffectsQueue>())
 		{
 			ReceiveInstantMovementEffects(TimeStep, InstantMovementEffectsQueue);
 		}
+		// Some may come from local sim input, non networked, in the case of instant movement effects queued at non input producing end points (remotely controlled actor on a server or sim proxy on a client)
+		if (const FChaosNetInstantMovementEffectsQueue* LocalInstantMovementEffectsQueue = Simulation->GetLocalSimInput().FindDataByType<FChaosNetInstantMovementEffectsQueue>())
+		{
+			ReceiveInstantMovementEffects(TimeStep, LocalInstantMovementEffectsQueue);
+		}
+		// Same for layered moves (we need to generalize this to handle any abstract "sim command")
+		// Most will come from the networked input
+		if (const FChaosNetLayeredMovesQueue* LayeredMovesQueue = StartState.InputCmd.InputCollection.FindDataByType<FChaosNetLayeredMovesQueue>())
+		{
+			ReceiveLayeredMoves(TimeStep, LayeredMovesQueue);
+		}
+		// Some may come from local sim input, in the case of moves queued at non input producing end points (remotely controlled actor on a server or sim proxy on a client)
+		if (const FChaosNetLayeredMovesQueue* LocalLayeredMovesQueue = Simulation->GetLocalSimInput().FindDataByType<FChaosNetLayeredMovesQueue>())
+		{
+			ReceiveLayeredMoves(TimeStep, LocalLayeredMovesQueue);
+		}
 
 #if WITH_CHAOS_VISUAL_DEBUGGER
-		CaptureUnprocessedInstantMovementEffectsForCVD();
+		CaptureUnprocessedSimCommandsForCVD();
 #endif
 
 		TStrongObjectPtr<UImmediateMovementModeTransition> QueuedModeTransition = QueuedModeTransitionWeakPtr.Pin();
@@ -377,7 +465,7 @@ namespace UE::ChaosMover
 			SubTimeStep.StepMs = TimeStep.StepMs - TotalUsedMs;		// TODO: convert this to an overridable function that can support MaxStepTime, MaxIterations, etc.
 
 			// Transfer any queued moves into the starting state. They'll be started during the move generation.
-			FlushQueuedMovesToGroup(WorkingSubstepStartData.SyncState.LayeredMoves);
+			FlushQueuedMovesToGroup(TimeStep, WorkingSubstepStartData.SyncState.LayeredMoves);
 			OutputState.SyncState.LayeredMoves = WorkingSubstepStartData.SyncState.LayeredMoves;
 
 			FlushQueuedModifiersToGroup(WorkingSubstepStartData.SyncState.MovementModifiers);
@@ -630,24 +718,26 @@ namespace UE::ChaosMover
 		}
 
 		// We verify no effects are left that should have been applied
-#if !defined(BUILD_SHIPPING) || !BUILD_SHIPPING
+#if !UE_BUILD_SHIPPING
 		if (const FChaosMoverSimulationDefaultInputs* SimInputs = Simulation->GetLocalSimInput().FindDataByType<FChaosMoverSimulationDefaultInputs>())
 		{
 			for (const FChaosScheduledInstantMovementEffect& QueuedEffect : QueuedInstantEffects)
 			{
-				ensureMsgf(!QueuedEffect.ScheduledEffect.ShouldExecuteAtFrame(TimeStep.ServerFrame),
+				ensureMsgf(!QueuedEffect.ScheduledEffect.SchedulingInfo.ShouldExecuteAtFrame(TimeStep.ServerFrame),
 					TEXT("(%s) An Instant Movement Effect that should have been applied at ServerFrame %d was still in the queue after the simulation was ticked at ServerFrame %d. Effect: %s"),
-					*ToString(SimInputs->OwningActor->GetNetMode()), QueuedEffect.ScheduledEffect.ExecutionServerFrame, TimeStep.ServerFrame, * QueuedEffect.ScheduledEffect.Effect->ToSimpleString());
+					NETMODE_TO_STR(SimInputs),
+					QueuedEffect.ScheduledEffect.SchedulingInfo.ServerExecutionTime.FrameCount,
+					TimeStep.ServerFrame,
+					EFFECTSHAREDPTR_TO_STR(QueuedEffect.ScheduledEffect.Effect));
 			}
 		}
-#endif // !defined(BUILD_SHIPPING) || !BUILD_SHIPPING
+#endif // !UE_BUILD_SHIPPING
 	}
 
 	void FMoverStateMachine::OnSimulationRollback(const FMoverTimeStep& NewTimeStep, const FMoverSyncState& InvalidSyncState, const FMoverSyncState& NewSyncState)
 	{
 		QUICK_SCOPE_CYCLE_COUNTER(MoverStateMachine_OnSimulationRollback);
 
-		const int32 PreviousInternalServerFrame = InternalSimTime.ServerFrame;
 		InternalSimTime = NewTimeStep;
 
 		ClearQueuedMode();
@@ -657,32 +747,11 @@ namespace UE::ChaosMover
 			QueueNextModeInternal(NewSyncState.MovementMode, /*bShouldReenter*/ false, /*bIsFromRollback*/ true);
 		}
 
-		QueuedLayeredMoves.Empty();
+		RollBackQueuedLayeredMoves(NewTimeStep.ServerFrame);
 		QueuedMovementModifiers.Empty();
+		RollBackQueuedInstantMovementEffects(NewTimeStep.ServerFrame);
 		InstantMovementEffectsIDHistory.Rollback(NewTimeStep.ServerFrame);
-
-		for (int EffectIndex = 0; EffectIndex < QueuedInstantEffects.Num();)
-		{
-			const FChaosScheduledInstantMovementEffect& QueuedEffect = QueuedInstantEffects[EffectIndex];
-			// We roll back all effects which were issued after PreviousInternalServerFrame, except those marked as not rolling back
-			// (typically those coming from the game thread that is not being resimulated)
-			if (QueuedEffect.IssuanceServerFrame >= PreviousInternalServerFrame && QueuedEffect.bShouldRollBack)
-			{
-#if !defined(BUILD_SHIPPING) || !BUILD_SHIPPING
-				if (const FChaosMoverSimulationDefaultInputs* SimInputs = Simulation->GetLocalSimInput().FindDataByType<FChaosMoverSimulationDefaultInputs>())
-				{
-					UE_LOG(LogChaosMover, Verbose, TEXT("(%s) Rolling back Instant Effect issued at frame %d scheduled for frame %d at frame %d: %s"),
-						*ToString(SimInputs->OwningActor->GetNetMode()), QueuedEffect.IssuanceServerFrame, QueuedEffect.ScheduledEffect.ExecutionServerFrame, NewTimeStep.ServerFrame, QueuedEffect.ScheduledEffect.Effect.IsValid() ? *QueuedEffect.ScheduledEffect.Effect->ToSimpleString() : TEXT("INVALID INSTANT EFFECT"));
-				}
-#endif // !defined(BUILD_SHIPPING) || !BUILD_SHIPPING
-
-				QueuedInstantEffects.RemoveAt(EffectIndex);
-			}
-			else
-			{
-				EffectIndex++;
-			}
-		}
+		LayeredMovesIDHistory.Rollback(NewTimeStep.ServerFrame);
 
 		FMovementModifierParams_Async ModifierParams(Simulation.Get(), &NewSyncState, &NewTimeStep);
 
@@ -732,6 +801,63 @@ namespace UE::ChaosMover
 		}
 	}
 
+	void FMoverStateMachine::RollBackQueuedInstantMovementEffects(int32 NewServerFrame)
+	{
+		for (int EffectIndex = 0; EffectIndex < QueuedInstantEffects.Num();)
+		{
+			const FChaosScheduledInstantMovementEffect& QueuedEffect = QueuedInstantEffects[EffectIndex];
+			// We roll back all effects which were issued on or after NewServerFrame, except those marked as not rolling back
+			// (typically those coming from the game thread that is not being resimulated)
+			if (QueuedEffect.ScheduledEffect.SchedulingInfo.ShouldRollBackAtFrame(NewServerFrame) && QueuedEffect.bShouldRollBack)
+			{
+#if !UE_BUILD_SHIPPING
+				if (const FChaosMoverSimulationDefaultInputs* SimInputs = Simulation->GetLocalSimInput().FindDataByType<FChaosMoverSimulationDefaultInputs>())
+				{
+					UE_LOGF(LogChaosMover, Verbose, "(%ls) Rolling back Instant Movement Effect scheduled [%ls] at frame %d: %ls",
+						NETMODE_TO_STR(SimInputs),
+						*QueuedEffect.ScheduledEffect.SchedulingInfo.ToString(),
+						NewServerFrame,
+						EFFECTSHAREDPTR_TO_STR(QueuedEffect.ScheduledEffect.Effect));
+				}
+#endif // !UE_BUILD_SHIPPING
+
+				QueuedInstantEffects.RemoveAt(EffectIndex);
+			}
+			else
+			{
+				EffectIndex++;
+			}
+		}
+	}
+
+	void FMoverStateMachine::RollBackQueuedLayeredMoves(int32 NewServerFrame)
+	{
+		for (int MoveIndex = 0; MoveIndex < QueuedLayeredMoves.Num();)
+		{
+			const FChaosScheduledLayeredMove& QueuedMove = QueuedLayeredMoves[MoveIndex];
+			// We roll back all moves which were issued on or after NewServerFrame, except those marked as not rolling back
+			// (typically those coming from the game thread that is not being resimulated)
+			if (QueuedMove.ScheduledLayeredMove.SchedulingInfo.ShouldRollBackAtFrame(NewServerFrame) && QueuedMove.bShouldRollBack)
+			{
+#if !UE_BUILD_SHIPPING
+				if (const FChaosMoverSimulationDefaultInputs* SimInputs = Simulation->GetLocalSimInput().FindDataByType<FChaosMoverSimulationDefaultInputs>())
+				{
+					UE_LOGF(LogChaosMover, Verbose, "(%ls) Rolling back Layered Move scheduled [%ls] at frame %d: %ls",
+						NETMODE_TO_STR(SimInputs),
+						*QueuedMove.ScheduledLayeredMove.SchedulingInfo.ToString(),
+						NewServerFrame,
+						MOVESHAREDPTR_TO_STR(QueuedMove.ScheduledLayeredMove.Move));
+				}
+#endif // !UE_BUILD_SHIPPING
+
+				QueuedLayeredMoves.RemoveAt(MoveIndex);
+			}
+			else
+			{
+				MoveIndex++;
+			}
+		}
+	}
 
 	const TWeakObjectPtr<UBaseMovementMode> FMoverStateMachine::GetCurrentMode() const
 	{
@@ -763,13 +889,21 @@ namespace UE::ChaosMover
 		return nullptr;
 	}
 
-	void FMoverStateMachine::QueueLayeredMove(TSharedPtr<FLayeredMoveBase> Move)
+	void FMoverStateMachine::QueueLayeredMove(const FChaosScheduledLayeredMove& ScheduledMove)
 	{
-		QueuedLayeredMoves.Add(Move);
+		Chaos::EnsureIsInPhysicsThreadContext();
+
+		QueuedLayeredMoves.Add(ScheduledMove);
+
+#if WITH_CHAOS_VISUAL_DEBUGGER
+		TraceLayeredMove(ScheduledMove);
+#endif
 	}
 
 	void FMoverStateMachine::QueueInstantMovementEffect(const FChaosScheduledInstantMovementEffect& ScheduledEffect)
 	{
+		Chaos::EnsureIsInPhysicsThreadContext();
+
 		QueuedInstantEffects.Add(ScheduledEffect);
 
 #if WITH_CHAOS_VISUAL_DEBUGGER
@@ -859,16 +993,23 @@ namespace UE::ChaosMover
 		}
 	}
 
-	void FMoverStateMachine::FlushQueuedMovesToGroup(FLayeredMoveGroup& Group)
+	void FMoverStateMachine::FlushQueuedMovesToGroup(const FMoverTimeStep& TimeStep, FLayeredMoveGroup& Group)
 	{
-		if (!QueuedLayeredMoves.IsEmpty())
+		for (int MoveIndex = 0; MoveIndex < QueuedLayeredMoves.Num(); )
 		{
-			for (TSharedPtr<FLayeredMoveBase>& QueuedMove : QueuedLayeredMoves)
+			const FScheduledLayeredMove& QueuedMove = QueuedLayeredMoves[MoveIndex].ScheduledLayeredMove;
+			bool bShouldFlushMove = QueuedMove.SchedulingInfo.ShouldExecuteAtFrame(TimeStep.ServerFrame);
+
+			if (bShouldFlushMove)
 			{
-				Group.QueueLayeredMove(QueuedMove);
+				// First enqueue, then remove, to make sure any shared pointer ref doesn't go to 0 and accidentally delete underlying objects. Don't ask me how I know.
+				Group.QueueLayeredMove(QueuedMove.Move);
+				QueuedLayeredMoves.RemoveAt(MoveIndex);
+			}
+			else
+			{
+				MoveIndex++;
 			}
-
-			QueuedLayeredMoves.Empty();
 		}
 	}
 
@@ -886,17 +1027,20 @@ namespace UE::ChaosMover
 		for (int EffectIndex = 0; EffectIndex < QueuedInstantEffects.Num(); )
 		{
 			const FScheduledInstantMovementEffect& ScheduledEffect = QueuedInstantEffects[EffectIndex].ScheduledEffect;
-			if (ScheduledEffect.ShouldExecuteAtFrame(TimeStep.ServerFrame))
+			if (ScheduledEffect.SchedulingInfo.ShouldExecuteAtFrame(TimeStep.ServerFrame))
 			{
 				bInstantMovementEffectApplied |= ScheduledEffect.Effect.IsValid() && ScheduledEffect.Effect->ApplyMovementEffect_Async(EffectParams_Async, OutputState);
 
-#if !defined(BUILD_SHIPPING) || !BUILD_SHIPPING
+#if !UE_BUILD_SHIPPING
 				if (const FChaosMoverSimulationDefaultInputs* SimInputs = Simulation->GetLocalSimInput().FindDataByType<FChaosMoverSimulationDefaultInputs>())
 				{
-					UE_LOG(LogChaosMover, Verbose, TEXT("(%s) Applying Instant Effect scheduled for frame %d at frame %d: %s"),
-						*ToString(SimInputs->OwningActor->GetNetMode()), ScheduledEffect.ExecutionServerFrame, TimeStep.ServerFrame, ScheduledEffect.Effect.IsValid() ? *ScheduledEffect.Effect->ToSimpleString() : TEXT("INVALID INSTANT EFFECT"));
+					UE_LOGF(LogChaosMover, Verbose, "(%ls) Applying Instant Effect scheduled [%ls] at frame %d: %ls",
+						NETMODE_TO_STR(SimInputs),
+						*ScheduledEffect.SchedulingInfo.ToString(),
+						TimeStep.ServerFrame,
+						EFFECTSHAREDPTR_TO_STR(ScheduledEffect.Effect));
 				}
-#endif // !defined(BUILD_SHIPPING) || !BUILD_SHIPPING
+#endif // !UE_BUILD_SHIPPING
 
 				// We remove this effect since it has been applied
 				QueuedInstantEffects.RemoveAt(EffectIndex);
@@ -904,13 +1048,16 @@ namespace UE::ChaosMover
 			else
 			{
 				EffectIndex++;
-#if !defined(BUILD_SHIPPING) || !BUILD_SHIPPING
+#if !UE_BUILD_SHIPPING
 				if (const FChaosMoverSimulationDefaultInputs* SimInputs = Simulation->GetLocalSimInput().FindDataByType<FChaosMoverSimulationDefaultInputs>())
 				{
-					UE_LOG(LogChaosMover, Verbose, TEXT("(%s) SKIPPING Instant Effect scheduled for frame %d at frame %d: %s"),
-						*ToString(SimInputs->OwningActor->GetNetMode()), ScheduledEffect.ExecutionServerFrame, TimeStep.ServerFrame, ScheduledEffect.Effect.IsValid() ? *ScheduledEffect.Effect->ToSimpleString() : TEXT("INVALID INSTANT EFFECT"));
+					UE_LOGF(LogChaosMover, Verbose, "(%ls) SKIPPING Instant Effect scheduled [%ls] at frame %d: %ls",
+						NETMODE_TO_STR(SimInputs),
+						*ScheduledEffect.SchedulingInfo.ToString(),
+						TimeStep.ServerFrame,
+						EFFECTSHAREDPTR_TO_STR(ScheduledEffect.Effect));
 				}
-#endif // !defined(BUILD_SHIPPING) || !BUILD_SHIPPING
+#endif // !UE_BUILD_SHIPPING
 			}
 		}
 
@@ -982,41 +1129,61 @@ namespace UE::ChaosMover
 	}
 
 #if WITH_CHAOS_VISUAL_DEBUGGER
-	void FMoverStateMachine::CaptureUnprocessedInstantMovementEffectsForCVD()
-	{
-		// Record the list of instant movement effects we ended up with at the start of OnSimulationTick
-		// These could be a combination of effects we will process on the current frame or effects that are scheduled to be processed later.
-		// The fact that we have scheduled effects that stay in the queue for a while is the reason we can't simply capture effects for recording in QueueInstantMovementEffect
-		// Some of these effects come from the FChaosNetInstantMovementEffectsQueue found in StartState.InputCmd.InputCollection (transferred in ReceiveInstantMovementEffects)
-		// while others come from direct calls to QueueInstantMovementEffect on the simulation
-		// Since we want to also record the effects that are added during OnSimulationTick after this point, we set a flag so any calls to QueueInstantMovementEffect
-		// from now till post solve when we trace data to CVD are to be added to DebugInstantMovementEffectsQueue
+	void FMoverStateMachine::CaptureUnprocessedSimCommandsForCVD()
+	{
+		// Record the list of sim commands (instant movement effects, layered moves, etc.) we ended up with at the start of OnSimulationTick
+		// These could be a combination of commands we will process on the current frame or that are scheduled to be processed later.
+		// The fact that we have scheduled commands that stay queued for a while is the reason we can't simply capture them for recording in QueueInstantMovementEffect, QueueLayeredMove, etc.
+		// Taking effects as an example, some of these come from the FChaosNetInstantMovementEffectsQueue found in StartState.InputCmd.InputCollection (transferred in ReceiveInstantMovementEffects)
+		// while others come from direct (internal) calls to QueueInstantMovementEffect on the simulation while it steps
+		// Since we want to also record those added during OnSimulationTick after this point, we set a flag so any calls to QueueInstantMovementEffect
+		// from now till post solve when we trace data to CVD are to be added to DebugInstantMovementEffectsQueue (same for moves, etc.)
 		if (FChaosVisualDebuggerTrace::IsTracing())
 		{
 			if (TStrongObjectPtr<UChaosMoverSimulation> SimStrongObjPtr = Simulation.Pin())
 			{
-				FChaosNetInstantMovementEffectsQueue& DebugInstantMovementEffectsQueue = SimStrongObjPtr->GetDebugSimData().FindOrAddMutableDataByType<FChaosNetInstantMovementEffectsQueue>();
+				FMoverDataCollection& DebugSimData = SimStrongObjPtr->GetDebugSimData();
+
+				// Instant Movement Effects
+				FChaosNetInstantMovementEffectsQueue& DebugInstantMovementEffectsQueue = DebugSimData.FindOrAddMutableDataByType<FChaosNetInstantMovementEffectsQueue>();
 				DebugInstantMovementEffectsQueue = QueuedInstantEffects;
-				bCaptureInstantMovementEffectsForCVD = true;
+				
+				// Layered Moves
+				FChaosNetLayeredMovesQueue& DebugLayeredMovesQueue = DebugSimData.FindOrAddMutableDataByType<FChaosNetLayeredMovesQueue>();
+				DebugLayeredMovesQueue = QueuedLayeredMoves;
+
+				bCaptureSimCommandsForCVD = true;
 			}
 		}
 	}
 
-	void FMoverStateMachine::TraceInstantMovementEffect(const FChaosScheduledInstantMovementEffect& ScheduledEffect)
+	void FMoverStateMachine::TraceInstantMovementEffect(const FChaosScheduledInstantMovementEffect& InScheduledEffect)
 	{
-		if (bCaptureInstantMovementEffectsForCVD && FChaosVisualDebuggerTrace::IsTracing())
+		if (bCaptureSimCommandsForCVD && FChaosVisualDebuggerTrace::IsTracing())
 		{
 			if (TStrongObjectPtr<UChaosMoverSimulation> SimStrongObjPtr = Simulation.Pin())
 			{
 				FChaosNetInstantMovementEffectsQueue& DebugInstantMovementEffectsQueue = SimStrongObjPtr->GetDebugSimData().FindOrAddMutableDataByType<FChaosNetInstantMovementEffectsQueue>();
-				DebugInstantMovementEffectsQueue.Add(ScheduledEffect.ScheduledEffect, ScheduledEffect.IssuanceServerFrame, 0xFF);
+				DebugInstantMovementEffectsQueue.Add(InScheduledEffect.ScheduledEffect, InScheduledEffect.bShouldRollBack, 0xFF);
+			}
+		}
+	}
+
+	void FMoverStateMachine::TraceLayeredMove(const FChaosScheduledLayeredMove& InScheduledMove)
+	{
+		if (bCaptureSimCommandsForCVD && FChaosVisualDebuggerTrace::IsTracing())
+		{
+			if (TStrongObjectPtr<UChaosMoverSimulation> SimStrongObjPtr = Simulation.Pin())
+			{
+				FChaosNetLayeredMovesQueue& DebugLayeredMovesQueue = SimStrongObjPtr->GetDebugSimData().FindOrAddMutableDataByType<FChaosNetLayeredMovesQueue>();
+				DebugLayeredMovesQueue.Add(InScheduledMove.ScheduledLayeredMove, InScheduledMove.bShouldRollBack, 0xFF);
 			}
 		}
 	}
 
 	void FMoverStateMachine::OnEndTraceMoverData()
 	{
-		bCaptureInstantMovementEffectsForCVD = false;
+		bCaptureSimCommandsForCVD = false;
 	}
 #endif
 
diff --git a/Engine/Plugins/Experimental/ChaosMover/Source/ChaosMover/Public/ChaosMover/Backends/ChaosMoverBackend.h b/Engine/Plugins/Experimental/ChaosMover/Source/ChaosMover/Public/ChaosMover/Backends/ChaosMoverBackend.h
--- a/Engine/Plugins/Experimental/ChaosMover/Source/ChaosMover/Public/ChaosMover/Backends/ChaosMoverBackend.h
+++ b/Engine/Plugins/Experimental/ChaosMover/Source/ChaosMover/Public/ChaosMover/Backends/ChaosMoverBackend.h
@@ -39,8 +39,10 @@ public:
 
 	// IMoverBackendLiaisonInterface
 	CHAOSMOVER_API virtual bool IsAsync() const override;
+	CHAOSMOVER_API virtual bool IsFixedDt() const override;
 	CHAOSMOVER_API virtual bool ShouldResim() const override;
 	CHAOSMOVER_API virtual float GetEventSchedulingMinDelaySeconds() const override;
+	CHAOSMOVER_API virtual FMoverTime GetScheduledNetworkTime(const FMoverTime& Time) const override;
 	CHAOSMOVER_API virtual double GetCurrentSimTimeMs();
 	CHAOSMOVER_API virtual int32 GetCurrentSimFrame();
 
@@ -107,16 +109,37 @@ protected:
 	bool bWantsDestroySim = false;
 	bool bWantsCreateSim = true;
 
-	// Transfers queued instant movement effects from the mover component to the simulation input.
+	// Transfers queued instant movement effects from the mover component to the simulation (via networked input or local sim input).
+	void InjectInstantMovementEffectsIntoSim(UE::ChaosMover::FSimulationInputData& OutInputData);
+
+	// Transfers queued instant movement effects from the mover component to the (networked) simulation input.
 	// This is used in the input producing case.
-	void InjectInstantMovementEffectsIntoInput(const FMoverTimeStep& TimeStep, UE::ChaosMover::FSimulationInputData& InputData);
+	void InjectInstantMovementEffectsIntoInput(UE::ChaosMover::FSimulationInputData& OutInputData);
 	
-	// Transfers queued instant movement effects from the mover component to the simulation directly, bypassing the input.
+	// Transfers queued instant movement effects from the mover component to the simulation directly, bypassing the networked input.
 	// This is used in the non input producing case.
-	void InjectInstantMovementEffectsIntoSim(const FMoverTimeStep& TimeStep);
+	void InjectInstantMovementEffectsIntoLocalSimInput();
+
+	// Transfers queued layered moves from the mover component to the simulation (via networked input or local sim input).
+	void InjectLayeredMovesIntoSim(UE::ChaosMover::FSimulationInputData& OutInputData);
+
+	// Transfers queued layered moves from the mover component to the (networked) simulation input.
+	// This is used in the input producing case.
+	void InjectLayeredMovesIntoInput(UE::ChaosMover::FSimulationInputData& OutInputData);
+
+	// Transfers queued layered moves from the mover component to the simulation directly, bypassing the networked input.
+	// This is used in the non input producing case.
+	void InjectLayeredMovesIntoLocalSimInput();
 
 private:
-	// Unique ID counter for instant movement effects.
-	// Unique IDs need to be generated because input can be repeated and we need to prevent multiple application of a given effect only queued once.
-	uint8 NextInstantMovementEffectUniqueID = 0xFF;
-};
\ No newline at end of file
+
+	void ConvertToSerializable(const TArray<FScheduledInstantMovementEffect>& ScheduledInstantMovementEffects, FChaosNetInstantMovementEffectsQueue& OutNetInstantMovementEffectsQueue);
+	void ConvertToSerializable(const TArray<FScheduledLayeredMove>& ScheduledLayeredMoves, FChaosNetLayeredMovesQueue& OutNetLayeredMovesQueue);
+
+	// Returns the next unique ID used to prevent multiple application of the same sim command
+	uint8 GetNextUniqueSimCommandID();
+
+	// Unique ID counter for sim commands.
+	// Unique IDs need to be generated because input can be repeated and we need to prevent multiple application of a given command only queued once.
+	uint8 NextUniqueSimCommandID = 0xFF;
+};
diff --git a/Engine/Plugins/Experimental/ChaosMover/Source/ChaosMover/Public/ChaosMover/ChaosMoverConsoleVariables.h b/Engine/Plugins/Experimental/ChaosMover/Source/ChaosMover/Public/ChaosMover/ChaosMoverConsoleVariables.h
--- a/Engine/Plugins/Experimental/ChaosMover/Source/ChaosMover/Public/ChaosMover/ChaosMoverConsoleVariables.h
+++ b/Engine/Plugins/Experimental/ChaosMover/Source/ChaosMover/Public/ChaosMover/ChaosMoverConsoleVariables.h
@@ -15,6 +15,7 @@ namespace UE::ChaosMover
 		CHAOSMOVER_API extern bool bSkipGenerateMoveIfOverridden;
 		CHAOSMOVER_API extern bool bDisableResimDuplicateEventChecking;
 		CHAOSMOVER_API extern int32 InstantMovementEffectIDHistorySize;
+		CHAOSMOVER_API extern int32 LayeredMoveIDHistorySize;
 		CHAOSMOVER_API extern int32 FrameDifferenceLeniencyForEventTimeComparison;
 		CHAOSMOVER_API extern float FallingCheckRelativeSpeedLimit;
 
diff --git a/Engine/Plugins/Experimental/ChaosMover/Source/ChaosMover/Public/ChaosMover/ChaosMoverSimulation.h b/Engine/Plugins/Experimental/ChaosMover/Source/ChaosMover/Public/ChaosMover/ChaosMoverSimulation.h
--- a/Engine/Plugins/Experimental/ChaosMover/Source/ChaosMover/Public/ChaosMover/ChaosMoverSimulation.h
+++ b/Engine/Plugins/Experimental/ChaosMover/Source/ChaosMover/Public/ChaosMover/ChaosMoverSimulation.h
@@ -52,6 +52,15 @@ public:
 	// Queue a scheduled instant movement effect
 	CHAOSMOVER_API void QueueInstantMovementEffect(const FScheduledInstantMovementEffect& ScheduledInstantMovementEffect, bool bShouldRollBack = true);
 
+	UFUNCTION(BlueprintCallable, CustomThunk, Category = Mover, meta = (CustomStructureParam = "LayeredMoveAsRawData", AllowAbstract = "false", DisplayName = "Queue Layered Move"))
+	void K2_QueueLayeredMove(UPARAM(DisplayName = "Layered Move") const int32& LayeredMoveAsRawData);
+	DECLARE_FUNCTION(execK2_QueueLayeredMove);
+
+	// Queue a Layered Move to take place at the end of this frame or start of the next subtick - whichever happens first
+	CHAOSMOVER_API void QueueLayeredMove(TSharedPtr<FLayeredMoveBase> LayeredMove, bool bShouldRollBack = true);
+	// Queue a scheduled layered move
+	CHAOSMOVER_API void QueueLayeredMove(const FScheduledLayeredMove& ScheduledLayeredMove, bool bShouldRollBack = true);
+
 	UFUNCTION(BlueprintCallable, CustomThunk, Category = Mover, meta = (CustomStructureParam = "MoveAsRawData", AllowAbstract = "false", DisplayName = "Queue Movement Modifier"))
 	CHAOSMOVER_API FMovementModifierHandle K2_QueueMovementModifier(UPARAM(DisplayName = "Movement Modifier") const int32& MoveAsRawData);
 	DECLARE_FUNCTION(execK2_QueueMovementModifier);
diff --git a/Engine/Plugins/Experimental/ChaosMover/Source/ChaosMover/Public/ChaosMover/ChaosMoverSimulationTypes.h b/Engine/Plugins/Experimental/ChaosMover/Source/ChaosMover/Public/ChaosMover/ChaosMoverSimulationTypes.h
--- a/Engine/Plugins/Experimental/ChaosMover/Source/ChaosMover/Public/ChaosMover/ChaosMoverSimulationTypes.h
+++ b/Engine/Plugins/Experimental/ChaosMover/Source/ChaosMover/Public/ChaosMover/ChaosMoverSimulationTypes.h
@@ -10,6 +10,7 @@
 #include "MoveLibrary/FloorQueryUtils.h"
 #include "MoveLibrary/WaterMovementUtils.h"
 #include "MoverSimulationTypes.h"
+#include "LayeredMove.h"
 #include "StructUtils/InstancedStruct.h"
 
 #include "ChaosMoverSimulationTypes.generated.h"
@@ -395,15 +396,12 @@ struct FStanceModifiedEventData : public FMoverSimulationEventData
 	EStanceMode NewStance = EStanceMode::Invalid;
 };
 
-// Version of a FScheduledInstantMovementEffect with Issuance Server Frame
+// Version of a FScheduledInstantMovementEffect bShouldRollBack flag
 USTRUCT()
 struct FChaosScheduledInstantMovementEffect
 {
 	GENERATED_BODY()
 
-	UPROPERTY(VisibleAnywhere, Category = "Mover")
-	int32 IssuanceServerFrame = INDEX_NONE;
-
 	// Whether this effect should rollback or not
 	// We do not net serialize this value since any effect received from the network should always be rolled back
 	// Only effects that were issued locally by the game thread should NOT roll back since resimulation does not run 
@@ -441,13 +439,11 @@ struct FChaosNetInstantMovementEffect
 
 	FScheduledInstantMovementEffect AsScheduledInstantMovementEffect() const
 	{
-		return FScheduledInstantMovementEffect(ExecutionServerFrame, /* ExecutionDateSeconds = */ 0.0, /* bIsAsyncMode = */ true, TSharedPtr<FInstantMovementEffect>(Effect.Get().Clone()));
+		return FScheduledInstantMovementEffect(
+			FMoverSchedulingInfo(FMoverTime(IssuanceServerFrame, /* TimeMs = */ 0.0), FMoverTime(ExecutionServerFrame, /* TimeMs = */ 0.0), /* bIsFixedDt = */ true),
+			TSharedPtr<FInstantMovementEffect>(ensure(Effect.IsValid()) ? Effect.Get().Clone() : nullptr));
 	}
 
-	// Eventually we want this to be a UPROPERTY so it can be displayed automatically in CVD
-	// However this causes a crash when CVD tries to load this property in SerializeTaggedProperty because
-	// it relies on the serialization of the struct class as a UObject, which is not supported by the type of archives used
-	// to send and receive Mover debug data to CVD
 	UPROPERTY(VisibleAnywhere, Category = ChaosMover)
 	TInstancedStruct<FInstantMovementEffect> Effect;
 };
@@ -460,7 +456,7 @@ struct FChaosNetInstantMovementEffectsQueue : public FMoverDataStructBase
 	UPROPERTY(VisibleAnywhere, Category = ChaosMover)
 	TArray<FChaosNetInstantMovementEffect> Effects;
 
-	void Add(const FScheduledInstantMovementEffect& Effect, int32 IssuanceServerFrame, uint8 UniqueID);
+	void Add(const FScheduledInstantMovementEffect& Effect, bool bShouldRollBack, uint8 UniqueID);
 
 	// This function is used for debugging (sending to CVD)
 	FChaosNetInstantMovementEffectsQueue& operator=(const TArray<FChaosScheduledInstantMovementEffect>& ScheduledInstantEffects);
@@ -475,6 +471,79 @@ struct FChaosNetInstantMovementEffectsQueue : public FMoverDataStructBase
 	virtual void Decay(float DecayAmount) override;
 };
 
+// Version of a FScheduledLayeredMove with bShouldRollBack flag
+USTRUCT()
+struct FChaosScheduledLayeredMove
+{
+	GENERATED_BODY()
+
+	// Whether this move should rollback or not
+	// We do not net serialize this value since any move received from the network should always be rolled back
+	// Only moves issued locally by the game thread should NOT roll back since resimulation does not run 
+	// game thread logic again and will fail to generate them again
+	UPROPERTY(VisibleAnywhere, Category = "Mover")
+	bool bShouldRollBack = true;
+
+	UPROPERTY(VisibleAnywhere, Category = ChaosMover)
+	FScheduledLayeredMove ScheduledLayeredMove;
+};
+
+USTRUCT()
+struct FChaosNetLayeredMove
+{
+	GENERATED_BODY()
+
+	UPROPERTY(VisibleAnywhere, Category = "Mover")
+	int32 IssuanceServerFrame = INDEX_NONE;
+
+	UPROPERTY(VisibleAnywhere, Category = ChaosMover)
+	int32 ExecutionServerFrame = INDEX_NONE;
+
+	UPROPERTY(VisibleAnywhere, Category = ChaosMover)
+	uint8 UniqueID = 0xFF;
+
+	// Whether this move should roll back or not
+	// We do not net serialize this value since any move received from the network should always be rolled back
+	// Only moves that were issued locally by the game thread should NOT roll back since resimulation does not run 
+	// game thread logic again and will fail to generate them again
+	UPROPERTY(VisibleAnywhere, Category = "Mover")
+	bool bShouldRollBack = true;
+
+	FScheduledLayeredMove AsScheduledLayeredMove() const
+	{
+		return FScheduledLayeredMove(
+			FMoverSchedulingInfo(FMoverTime(IssuanceServerFrame, /* TimeMs = */ 0.0), FMoverTime(ExecutionServerFrame, /* TimeMs = */ 0.0), /* bIsFixedDt = */ true),
+			TSharedPtr<FLayeredMoveBase>(ensure(Move.IsValid()) ? Move.Get().Clone() : nullptr));
+	}
+
+	UPROPERTY(VisibleAnywhere, Category = ChaosMover)
+	TInstancedStruct<FLayeredMoveBase> Move;
+};
+
+USTRUCT()
+struct FChaosNetLayeredMovesQueue : public FMoverDataStructBase
+{
+	GENERATED_BODY()
+	
+	UPROPERTY(VisibleAnywhere, Category = ChaosMover)
+	TArray<FChaosNetLayeredMove> Moves;
+
+	void Add(const FScheduledLayeredMove& InScheduledLayeredMove, bool bShouldRollBack, uint8 UniqueID);
+
+	// This function is used for debugging (sending to CVD)
+	FChaosNetLayeredMovesQueue& operator=(const TArray<FChaosScheduledLayeredMove>& LayeredMovesQueue);
+
+	virtual FMoverDataStructBase* Clone() const override;
+	virtual bool NetSerialize(FArchive& Ar, UPackageMap* Map, bool& bOutSuccess) override;
+	virtual UScriptStruct* GetScriptStruct() const override;
+	virtual void ToString(FAnsiStringBuilderBase& Out) const override;
+	virtual bool ShouldReconcile(const FMoverDataStructBase& AuthorityState) const override;
+	virtual void Interpolate(const FMoverDataStructBase& From, const FMoverDataStructBase& To, float Pct) override;
+	virtual void Merge(const FMoverDataStructBase& From) override;
+	virtual void Decay(float DecayAmount) override;
+};
+
+
 // This instant movement effect is ChaosMover only and executes AsyncFunction on the issuing end point only
 // Being local only, this instant movement effect does not serialize anything, doesn't override NetSerialize
 USTRUCT()
diff --git a/Engine/Plugins/Experimental/ChaosMover/Source/ChaosMover/Public/ChaosMover/ChaosMoverStateMachine.h b/Engine/Plugins/Experimental/ChaosMover/Source/ChaosMover/Public/ChaosMover/ChaosMoverStateMachine.h
--- a/Engine/Plugins/Experimental/ChaosMover/Source/ChaosMover/Public/ChaosMover/ChaosMoverStateMachine.h
+++ b/Engine/Plugins/Experimental/ChaosMover/Source/ChaosMover/Public/ChaosMover/ChaosMoverStateMachine.h
@@ -55,7 +55,7 @@ namespace UE::ChaosMover
 		CHAOSMOVER_API const TWeakObjectPtr<UBaseMovementMode> FindMovementMode(FName ModeName) const;
 		CHAOSMOVER_API TWeakObjectPtr<UBaseMovementMode> FindMovementMode_Mutable(FName ModeName);
 
-		CHAOSMOVER_API void QueueLayeredMove(TSharedPtr<FLayeredMoveBase> Move);
+		CHAOSMOVER_API void QueueLayeredMove(const FChaosScheduledLayeredMove& ScheduledMove);
 
 		CHAOSMOVER_API void QueueInstantMovementEffect(const FChaosScheduledInstantMovementEffect& ScheduledEffect);
 
@@ -73,7 +73,8 @@ namespace UE::ChaosMover
 		void SetOwnerActorLocalNetRole(ENetRole InOwnerActorLocalNetRole);
 
 #if WITH_CHAOS_VISUAL_DEBUGGER
-		void CaptureUnprocessedInstantMovementEffectsForCVD();
+		// Capture any simulation commands (instant movement effects, layered moves, etc.) that have not yet been processed to be sent to CVD
+		void CaptureUnprocessedSimCommandsForCVD();
 		void OnEndTraceMoverData();
 #endif
 
@@ -96,10 +97,10 @@ namespace UE::ChaosMover
 		FName DefaultModeName = NAME_None;
 		FName CurrentModeName = NAME_None;
 
-		/** Moves that are queued to be added to the simulation at the start of the next sim subtick */
-		TArray<TSharedPtr<FLayeredMoveBase>> QueuedLayeredMoves;
+		/** Moves that are queued to be executed when their scheduled time comes (earliest is at the start of the next sim subtick) */
+		TArray<FChaosScheduledLayeredMove> QueuedLayeredMoves;
 
-		/** Effects that are queued to be applied to the simulation at the start of the next sim subtick or at the end of this tick */
+		/** Effects that are queued to be applied when their schduled time comes (earliest at the start of the next sim subtick or at the end of the ongoing tick if queued mid tick) */
 		TArray<FChaosScheduledInstantMovementEffect> QueuedInstantEffects;
 
 		/** Modifiers that are queued to be added to the simulation at the start of the next sim subtick. */
@@ -117,12 +118,15 @@ namespace UE::ChaosMover
 	private:
 		void ConstructDefaultModes();
 		void AdvanceToNextMode(const FMoverTimeStep& TimeStep);
-		void FlushQueuedMovesToGroup(FLayeredMoveGroup& Group);
+		void FlushQueuedMovesToGroup(const FMoverTimeStep& TimeStep, FLayeredMoveGroup& Group);
 		void FlushQueuedModifiersToGroup(FMovementModifierGroup& ModifierGroup);
 		void FlushModifierCancellationsToGroup(FMovementModifierGroup& ActiveModifierGroup);
 		bool HasAnyInstantEffectsQueued(const FMoverTimeStep& TimeStep) const;
 		bool ApplyInstantEffects(const FMoverTickStartData& SubstepStartData, const FMoverTimeStep& SubTimeStep, FMoverSyncState& OutputState);
 		void ReceiveInstantMovementEffects(const FMoverTimeStep& TimeStep, const FChaosNetInstantMovementEffectsQueue* InstantMovementEffectsQueue);
+		void ReceiveLayeredMoves(const FMoverTimeStep& TimeStep, const FChaosNetLayeredMovesQueue* LayeredMovesQueue);
+		void RollBackQueuedInstantMovementEffects(int32 NewServerFrame);
+		void RollBackQueuedLayeredMoves(int32 NewServerFrame);
 	
 		FMoverTimeStep InternalSimTime;
 
@@ -137,14 +141,16 @@ namespace UE::ChaosMover
 		};
 
 		FIDHistory InstantMovementEffectsIDHistory;
+		FIDHistory LayeredMovesIDHistory;
 
 #if WITH_CHAOS_VISUAL_DEBUGGER
-		// If true, queued instant movement effects will be added to DebugSimData
-		// It is false from the last time we traced movement data to the time we call CaptureUnprocessedInstantMovementEffectsForCVD in FMoverStateMachine::OnSimulationTick
-		// which allows us to trace any unprocessed effects, including scheduled ones, that exist on that frame
+		// If true, queued simulation commands (instant movement effects. layered moves, etc.) will be added to DebugSimData
+		// It is false from the last time we traced movement data to the time we call CaptureUnprocessedSimCommandsForCVD in FMoverStateMachine::OnSimulationTick
+		// which allows us to trace any unprocessed sim commands, including scheduled ones, that exist on that frame
 		// It is true after that point until we trace, so that we can capture the ones queued after that, for instance those queued by movement modes
-		bool bCaptureInstantMovementEffectsForCVD = false;
+		bool bCaptureSimCommandsForCVD = false;
 		void TraceInstantMovementEffect(const FChaosScheduledInstantMovementEffect& ScheduledEffect);
+		void TraceLayeredMove(const FChaosScheduledLayeredMove& ScheduledMove);
 #endif
 	};
 }
\ No newline at end of file
diff --git a/Engine/Plugins/Experimental/Mover/Source/Mover/Private/DefaultMovementSet/InstantMovementEffects/BasicInstantMovementEffects.cpp b/Engine/Plugins/Experimental/Mover/Source/Mover/Private/DefaultMovementSet/InstantMovementEffects/BasicInstantMovementEffects.cpp
--- a/Engine/Plugins/Experimental/Mover/Source/Mover/Private/DefaultMovementSet/InstantMovementEffects/BasicInstantMovementEffects.cpp
+++ b/Engine/Plugins/Experimental/Mover/Source/Mover/Private/DefaultMovementSet/InstantMovementEffects/BasicInstantMovementEffects.cpp
@@ -47,7 +47,7 @@ bool FTeleportEffect::ApplyMovementEffect(FApplyMovementEffectParams& ApplyEffec
 	if (OwnerActor->TeleportTo(TargetLocation, FinalTargetRotation))
 	{
 		const FVector UpdatedLocation = UpdatedComponent->GetComponentLocation();
-#if !defined(BUILD_SHIPPING) || !BUILD_SHIPPING
+#if !UE_BUILD_SHIPPING
 		if (ShowTeleportDiffs)
 		{
 			if (!(UpdatedLocation - TargetLocation).IsNearlyZero())	// if it was adjusted, show the original error
@@ -56,7 +56,7 @@ bool FTeleportEffect::ApplyMovementEffect(FApplyMovementEffectParams& ApplyEffec
 			}
 			DrawDebugCapsule(OwnerActor->GetWorld(), UpdatedLocation, OwnerActor->GetSimpleCollisionHalfHeight(), OwnerActor->GetSimpleCollisionRadius(), FQuat::Identity, FColor(100, 100, 255), false, ShowTeleportDiffsLifetimeSecs);
 		}
-#endif // !defined(BUILD_SHIPPING) || !BUILD_SHIPPING
+#endif // !UE_BUILD_SHIPPING
 
 		FMoverDefaultSyncState& OutputSyncState = OutputState.SyncStateCollection.FindOrAddMutableDataByType<FMoverDefaultSyncState>();
 		OutputSyncState.SetTransforms_WorldSpace(UpdatedLocation,
@@ -77,12 +77,12 @@ bool FTeleportEffect::ApplyMovementEffect(FApplyMovementEffectParams& ApplyEffec
 		return true;
 	}
 
-#if !defined(BUILD_SHIPPING) || !BUILD_SHIPPING
+#if !UE_BUILD_SHIPPING
 	if (ShowTeleportDiffs)
 	{
 		DrawDebugCapsule(OwnerActor->GetWorld(), TargetLocation, OwnerActor->GetSimpleCollisionHalfHeight(), OwnerActor->GetSimpleCollisionRadius(), FQuat::Identity, FColor::Red, false, ShowTeleportDiffsLifetimeSecs);
 	}
-#endif // !defined(BUILD_SHIPPING) || !BUILD_SHIPPING
+#endif // !UE_BUILD_SHIPPING
 
 	ApplyEffectParams.OutputEvents.Add(MakeShared<FTeleportFailedEventData>(*ApplyEffectParams.TimeStep, PreviousLocation, PreviousRotation, TargetLocation, FQuat(FinalTargetRotation), ETeleportFailureReason::Reason_NotAvailable));
 
diff --git a/Engine/Plugins/Experimental/Mover/Source/Mover/Private/MovementModeStateMachine.cpp b/Engine/Plugins/Experimental/Mover/Source/Mover/Private/MovementModeStateMachine.cpp
--- a/Engine/Plugins/Experimental/Mover/Source/Mover/Private/MovementModeStateMachine.cpp
+++ b/Engine/Plugins/Experimental/Mover/Source/Mover/Private/MovementModeStateMachine.cpp
@@ -200,6 +200,12 @@ void UMovementModeStateMachine::OnSimulationTick(USceneComponent* UpdatedCompone
 	QueueInstantMovementEffects(ScheduledInstantMovementEffects);
 	MoverComp->ClearQueuedInstantMovementEffects();
 
+	// Copy queued layered moves from the mover component to the state machine
+	// After this, the state machine works on its own queue, into which it can enqueue layered moves while stepping
+	const TArray<FScheduledLayeredMove>& ScheduledLayeredMoves = MoverComp->GetQueuedLayeredMoves();
+	QueueLayeredMoves(ScheduledLayeredMoves);
+	MoverComp->ClearQueuedLayeredMoves();
+
 	if (!QueuedModeTransition->IsSet())
 	{
 		QueueNextMode(WorkingSubstepStartData.SyncState.MovementMode);
@@ -234,7 +240,7 @@ void UMovementModeStateMachine::OnSimulationTick(USceneComponent* UpdatedCompone
 		FlushTagCancellationsToSyncState(WorkingSubstepStartData.SyncState);
 
 		// Transfer any queued moves into the starting state. They'll be started during the move generation.
-		FlushQueuedMovesToGroup(WorkingSubstepStartData.SyncState.LayeredMoves);
+		FlushQueuedMovesToGroup(TimeStep, WorkingSubstepStartData.SyncState.LayeredMoves);
 		OutputState.SyncState.LayeredMoves = WorkingSubstepStartData.SyncState.LayeredMoves;
 		
 		ActivateQueuedMoves(WorkingSubstepStartData.SyncState.LayeredMoveInstances);
@@ -509,20 +515,14 @@ void UMovementModeStateMachine::OnSimulationRollback(const FMoverSyncState* Sync
 		SetModeImmediatelyInternal(SyncState->MovementMode, /*bShouldReenter*/ false, /*bIsFromRollback*/ true);
 	}
 
-	{
-		UE::TRWScopeLock QueueLock(LayeredMoveQueueLock, SLT_Write);
-		QueuedLayeredMoves.Empty();
-	}
+	RollBackQueuedLayeredMoves(NewBaseTimeStep);
 
 	{
 		UE::TRWScopeLock QueueLock(LayeredMoveQueueLock, SLT_Write);
 		QueuedLayeredMoveInstances.Empty();
 	}
 	
-	{
-		UE::TRWScopeLock QueueLock(InstantEffectsQueueLock, SLT_Write);
-		QueuedInstantEffects.Empty();
-	}
+	RollBackQueuedInstantMovementEffects(NewBaseTimeStep);
 
 	{
 		UE::TRWScopeLock QueueLock(ModifiersQueueLock, SLT_Write);
@@ -530,6 +530,43 @@ void UMovementModeStateMachine::OnSimulationRollback(const FMoverSyncState* Sync
 	}
 }
 
+void UMovementModeStateMachine::RollBackQueuedInstantMovementEffects(const FMoverTimeStep& NewBaseTimeStep)
+{
+	UE::TRWScopeLock QueueLock(InstantEffectsQueueLock, SLT_Write);
+
+	for (int EffectIndex = 0; EffectIndex < QueuedInstantEffects.Num();)
+	{
+		const FScheduledInstantMovementEffect& QueuedEffect = QueuedInstantEffects[EffectIndex];
+		// We roll back all effects which were issued on or after NewBaseTimeStep
+		if (QueuedEffect.SchedulingInfo.ShouldRollBackAtTimeStep(NewBaseTimeStep))
+		{
+			QueuedInstantEffects.RemoveAt(EffectIndex);
+		}
+		else
+		{
+			EffectIndex++;
+		}
+	}
+}
+
+void UMovementModeStateMachine::RollBackQueuedLayeredMoves(const FMoverTimeStep& NewBaseTimeStep)
+{
+	UE::TRWScopeLock QueueLock(LayeredMoveQueueLock, SLT_Write);
+	
+	for (int MoveIndex = 0; MoveIndex < QueuedLayeredMoves.Num();)
+	{
+		const FScheduledLayeredMove& QueuedMove = QueuedLayeredMoves[MoveIndex];
+		// We roll back all moves which were issued on or after NewBaseTimeStep
+		if (QueuedMove.SchedulingInfo.ShouldRollBackAtTimeStep(NewBaseTimeStep))
+		{
+			QueuedLayeredMoves.RemoveAt(MoveIndex);
+		}
+		else
+		{
+			MoveIndex++;
+		}
+	}
+}
 
 const UBaseMovementMode* UMovementModeStateMachine::GetCurrentMode() const
 {
@@ -561,10 +598,19 @@ const UBaseMovementMode* UMovementModeStateMachine::FindMovementMode(FName ModeN
 	return nullptr;
 }
 
-void UMovementModeStateMachine::QueueLayeredMove(TSharedPtr<FLayeredMoveBase> Move)
+void UMovementModeStateMachine::QueueInstantMovementEffect(const FScheduledInstantMovementEffect& ScheduledInstantMovementEffect)
 {
-	UE::TRWScopeLock QueueLock(LayeredMoveQueueLock, SLT_Write);
-	QueuedLayeredMoves.Add(Move);
+	UE::TRWScopeLock QueueLock(InstantEffectsQueueLock, SLT_Write);
+	QueuedInstantEffects.Add(ScheduledInstantMovementEffect);
+}
+
+void UMovementModeStateMachine::QueueInstantMovementEffects(const TArray<FScheduledInstantMovementEffect>& ScheduledInstantMovementEffects)
+{
+	UE::TRWScopeLock QueueLock(InstantEffectsQueueLock, SLT_Write);
+	for (const FScheduledInstantMovementEffect& ScheduledInstantMovementEffect : ScheduledInstantMovementEffects)
+	{
+		QueuedInstantEffects.Add(ScheduledInstantMovementEffect);
+	}
 }
 
 void UMovementModeStateMachine::QueueInstantMovementEffect_Internal(const FScheduledInstantMovementEffect& ScheduledInstantMovementEffect)
@@ -573,25 +619,31 @@ void UMovementModeStateMachine::QueueInstantMovementEffect_Internal(const FSched
 	QueueInstantMovementEffect(ScheduledInstantMovementEffect);
 }
 
-void UMovementModeStateMachine::QueueActiveLayeredMove(const TSharedPtr<FLayeredMoveInstance>& LayeredMove)
+void UMovementModeStateMachine::QueueLayeredMove(const FScheduledLayeredMove& ScheduledLayeredMove)
 {
 	UE::TRWScopeLock QueueLock(LayeredMoveQueueLock, SLT_Write);
-	QueuedLayeredMoveInstances.Add(LayeredMove);
+	QueuedLayeredMoves.Add(ScheduledLayeredMove);
 }
 
-void UMovementModeStateMachine::QueueInstantMovementEffect(const FScheduledInstantMovementEffect& ScheduledInstantMovementEffect)
+void UMovementModeStateMachine::QueueLayeredMoves(const TArray<FScheduledLayeredMove>& ScheduledLayeredMoves)
 {
-	UE::TRWScopeLock QueueLock(InstantEffectsQueueLock, SLT_Write);
-	QueuedInstantEffects.Add(ScheduledInstantMovementEffect);
+	UE::TRWScopeLock QueueLock(LayeredMoveQueueLock, SLT_Write);
+	for (const FScheduledLayeredMove& ScheduledLayeredMove : ScheduledLayeredMoves)
+	{
+		QueuedLayeredMoves.Add(ScheduledLayeredMove);
+	}
 }
 
-void UMovementModeStateMachine::QueueInstantMovementEffects(const TArray<FScheduledInstantMovementEffect>& ScheduledInstantMovementEffects)
+void UMovementModeStateMachine::QueueLayeredMove_Internal(const FScheduledLayeredMove& ScheduledLayeredMove)
 {
-	UE::TRWScopeLock QueueLock(InstantEffectsQueueLock, SLT_Write);
-	for (const FScheduledInstantMovementEffect& ScheduledInstantMovementEffect : ScheduledInstantMovementEffects)
-	{
-		QueuedInstantEffects.Add(ScheduledInstantMovementEffect);
-	}
+	ensure(!IsInGameThread());
+	QueueLayeredMove(ScheduledLayeredMove);
+}
+
+void UMovementModeStateMachine::QueueActiveLayeredMove(const TSharedPtr<FLayeredMoveInstance>& LayeredMove)
+{
+	UE::TRWScopeLock QueueLock(LayeredMoveQueueLock, SLT_Write);
+	QueuedLayeredMoveInstances.Add(LayeredMove);
 }
 
 FMovementModifierHandle UMovementModeStateMachine::QueueMovementModifier(TSharedPtr<FMovementModifierBase> Modifier)
@@ -677,9 +729,9 @@ void UMovementModeStateMachine::CancelFeaturesWithTag(FGameplayTag TagToCancel,
 
 	{
 		UE::TRWScopeLock QueueLock(LayeredMoveQueueLock, SLT_Write);
-		QueuedLayeredMoves.RemoveAll([TagToCancel, bRequireExactMatch](const TSharedPtr<FLayeredMoveBase>& LayeredMove)
+		QueuedLayeredMoves.RemoveAll([TagToCancel, bRequireExactMatch](const FScheduledLayeredMove& LayeredMove)
 			{
-				return (LayeredMove.IsValid() && LayeredMove->HasGameplayTag(TagToCancel, bRequireExactMatch));
+				return (LayeredMove.Move.IsValid() && LayeredMove.Move->HasGameplayTag(TagToCancel, bRequireExactMatch));
 			});
 	}
 
@@ -693,12 +745,11 @@ void UMovementModeStateMachine::CancelFeaturesWithTag(FGameplayTag TagToCancel,
 
 	// TODO: also support cancellation of queued instant effects if they end up supporting gameplay tags
 
-	// Request cancelation of any matching ACTIVE movement features during the next simulation tick
+	// Request cancellation of any matching ACTIVE movement features during the next simulation tick
 	{
 		UE::TRWScopeLock QueueLock(ModifierCancelQueueLock, SLT_Write);
 		TagCancellationRequests.Add(TPair<FGameplayTag, bool>(TagToCancel, bRequireExactMatch));
 	}
-
 }
 
 bool UMovementModeStateMachine::HasFeaturesWithTag(FGameplayTag TagToMatch, bool bRequireExactMatch) const
@@ -718,9 +769,9 @@ bool UMovementModeStateMachine::HasFeaturesWithTag(FGameplayTag TagToMatch, bool
 	if (!bResult)
 	{
 		UE::TRWScopeLock QueueLock(LayeredMoveQueueLock, SLT_ReadOnly);
-		const TSharedPtr<FLayeredMoveBase>* FindResult = QueuedLayeredMoves.FindByPredicate([TagToMatch, bRequireExactMatch](const TSharedPtr<FLayeredMoveBase>& LayeredMove)
+		const FScheduledLayeredMove* FindResult = QueuedLayeredMoves.FindByPredicate([TagToMatch, bRequireExactMatch](const FScheduledLayeredMove& LayeredMove)
 			{
-				return (LayeredMove.IsValid() && LayeredMove->HasGameplayTag(TagToMatch, bRequireExactMatch));
+				return (LayeredMove.Move.IsValid() && LayeredMove.Move->HasGameplayTag(TagToMatch, bRequireExactMatch));
 			});
 		bResult = FindResult != nullptr;
 	}
@@ -797,18 +848,35 @@ void UMovementModeStateMachine::AdvanceToNextMode()
 	}
 }
 
-void UMovementModeStateMachine::FlushQueuedMovesToGroup(FLayeredMoveGroup& Group)
+void UMovementModeStateMachine::FlushQueuedMovesToGroup(const FMoverTimeStep& TimeStep, FLayeredMoveGroup& Group)
 {
 	UE::TRWScopeLock QueueLock(LayeredMoveQueueLock, SLT_Write);
 
-	if (!QueuedLayeredMoves.IsEmpty())
+	for (int MoveIndex = 0; MoveIndex < QueuedLayeredMoves.Num(); )
 	{
-		for (TSharedPtr<FLayeredMoveBase>& QueuedMove : QueuedLayeredMoves)
+		const FScheduledLayeredMove& QueuedMove = QueuedLayeredMoves[MoveIndex];
+		bool bShouldFlushMove = false;
+		if (QueuedMove.SchedulingInfo.bIsFixedDt)
 		{
-			Group.QueueLayeredMove(QueuedMove);
+			bShouldFlushMove = QueuedMove.SchedulingInfo.ShouldExecuteAtFrame(TimeStep.ServerFrame);
+		}
+		else
+		{
+			UWorld* World = GetWorld();
+			AGameStateBase* GameState = World ? World->GetGameState() : nullptr;
+			double ServerTimeMs = GameState ? (GameState->GetServerWorldTimeSeconds() * 1000.0) : 0.0;
+			bShouldFlushMove = QueuedMove.SchedulingInfo.ShouldExecuteAtTimeMs(ServerTimeMs);
+		}
+		if (bShouldFlushMove)
+		{
+			// First enqueue, then remove, to make sure any shared pointer ref doesn't go to 0 and accidentally delete underlying objects. Don't ask me how I know.
+			Group.QueueLayeredMove(QueuedMove.Move);
+			QueuedLayeredMoves.RemoveAt(MoveIndex);
+		}
+		else
+		{
+			MoveIndex++;
 		}
-		
-		QueuedLayeredMoves.Empty();
 	}
 }
 
@@ -937,19 +1005,19 @@ bool UMovementModeStateMachine::ApplyInstantEffects(FApplyMovementEffectParams&
 	for (int EffectIndex = 0; EffectIndex < QueuedInstantEffects.Num(); )
 	{
 		const FScheduledInstantMovementEffect& QueuedEffect = QueuedInstantEffects[EffectIndex];
-		bool bshouldApplyEffect = false;
-		if (QueuedEffect.bIsFixedDt)
+		bool bShouldApplyEffect = false;
+		if (QueuedEffect.SchedulingInfo.bIsFixedDt)
 		{
-			bshouldApplyEffect = QueuedEffect.ShouldExecuteAtFrame(ApplyEffectParams.TimeStep->ServerFrame);
+			bShouldApplyEffect = QueuedEffect.SchedulingInfo.ShouldExecuteAtFrame(ApplyEffectParams.TimeStep->ServerFrame);
 		}
 		else
 		{
 			UWorld* World = GetWorld();
 			AGameStateBase* GameState = World ? World->GetGameState() : nullptr;
-			double ServerTime = GameState ? GameState->GetServerWorldTimeSeconds() : 0.0;
-			bshouldApplyEffect = QueuedEffect.ShouldExecuteAtTime(ServerTime);
+			double ServerTimeMs = GameState ? (GameState->GetServerWorldTimeSeconds() * 1000.0) : 0.0;
+			bShouldApplyEffect = QueuedEffect.SchedulingInfo.ShouldExecuteAtTimeMs(ServerTimeMs);
 		}
-		if (bshouldApplyEffect)
+		if (bShouldApplyEffect)
 		{
 			bInstantMovementEffectApplied |= QueuedEffect.Effect->ApplyMovementEffect(ApplyEffectParams, OutputState);
 			QueuedInstantEffects.RemoveAt(EffectIndex);
@@ -974,7 +1042,7 @@ void UMovementModeStateMachine::ProcessEvents(const TArray<TSharedPtr<FMoverSimu
 		{
 			ProcessSimulationEvent(*EventData);
 
-#if !defined(BUILD_SHIPPING) || !BUILD_SHIPPING
+#if !UE_BUILD_SHIPPING
 			ensureMsgf(IsInGameThread(), TEXT("Dispatching an event to the mover component from outside the game thread, this is not thread safe"));
 #endif
 			MoverComp->DispatchSimulationEvent(*EventData);
diff --git a/Engine/Plugins/Experimental/Mover/Source/Mover/Private/MoverComponent.cpp b/Engine/Plugins/Experimental/Mover/Source/Mover/Private/MoverComponent.cpp
--- a/Engine/Plugins/Experimental/Mover/Source/Mover/Private/MoverComponent.cpp
+++ b/Engine/Plugins/Experimental/Mover/Source/Mover/Private/MoverComponent.cpp
@@ -29,6 +29,7 @@
 #include "Components/SkeletalMeshComponent.h"
 #include "MotionWarpingComponent.h"
 #include "ChaosVisualDebugger/MoverCVDRuntimeTrace.h"
+#include "MoverSchedulingUtils.h"
 
 #if WITH_EDITOR
 #include "Misc/DataValidation.h"
@@ -327,12 +328,6 @@ void UMoverComponent::BeginPlay()
 	{
 		MovementMixer = NewObject<UMovementMixer>(this, TEXT("Default Movement Mixer"));
 	}
-
-	// Initialize the fixed delay for event scheduling
-	if (BackendLiaisonComp)
-	{
-		EventSchedulingMinDelaySeconds = BackendLiaisonComp->GetEventSchedulingMinDelaySeconds();
-	}
 }
 
 void UMoverComponent::BindProcessGeneratedMovement(FMover_ProcessGeneratedMovement ProcessGeneratedMovementEvent)
@@ -1670,13 +1665,65 @@ DEFINE_FUNCTION(UMoverComponent::execK2_QueueLayeredMove)
 	P_NATIVE_END;
 }
 
+void UMoverComponent::ScheduleLayeredMove(TSharedPtr<FLayeredMoveBase> LayeredMove)
+{
+	ensureMsgf(IsInGameThread(), TEXT("UMoverComponent::ScheduleLayeredMove should only be called from the game thread. Inspect code for incorrect calls."));
+	FMoverTimeStep TimeStep;
+	if (ensureMsgf(BackendLiaisonComp, TEXT("UMoverComponent::ScheduleLayeredMove was unexpectedly called with a null backend liaison component. The move will be ignored.")))
+	{
+		if (ensure(LayeredMove.IsValid()))
+		{
+			TimeStep.BaseSimTimeMs = BackendLiaisonComp->GetCurrentSimTimeMs();
+			TimeStep.ServerFrame = BackendLiaisonComp->GetCurrentSimFrame();
+			// TimeStep.StepMs is not used by UE::Mover::MakeScheduledMove so we don't set it here
+			QueueScheduledLayeredMove(UE::Mover::MakeScheduledMove(BackendLiaisonComp, TimeStep, LayeredMove));
+		}
+	}
+}
 
 void UMoverComponent::QueueLayeredMove(TSharedPtr<FLayeredMoveBase> LayeredMove)
-{	
-	if (ModeFSM)
+{
+	ensureMsgf(IsInGameThread(), TEXT("UMoverComponent::QueueLayeredMove(TSharedPtr<FLayeredMoveBase>) should only be called from the game thread. Inspect code for incorrect calls."));
+	FMoverTimeStep TimeStep;
+	if (ensureMsgf(BackendLiaisonComp, TEXT("UMoverComponent::QueueLayeredMove was unexpectedly called with a null backend liaison component. The move will be ignored.")))
+	{
+		if (ensure(LayeredMove.IsValid()))
+		{
+			TimeStep.BaseSimTimeMs = BackendLiaisonComp->GetCurrentSimTimeMs();
+			TimeStep.ServerFrame = BackendLiaisonComp->GetCurrentSimFrame();
+			// TimeStep.StepMs is not used by UE::Mover::ScheduleMove so we don't set it here
+			QueueScheduledLayeredMove(UE::Mover::MakeScheduledMoveForNextFrame(BackendLiaisonComp, TimeStep, LayeredMove));
+		}
+	}
+}
+
+void UMoverComponent::QueueScheduledLayeredMove(const FScheduledLayeredMove& ScheduledLayeredMove)
+{
+	// TODO Move QueueScheduledLayeredMove to UMoverSimulation and implement differently in sync or async mode
+	if (IsInGameThread())
+	{
+		QueuedLayeredMoves.Add(ScheduledLayeredMove);
+	}
+	else if (ensure(ModeFSM))
 	{
-		ModeFSM->QueueLayeredMove(LayeredMove);
+		ModeFSM->QueueLayeredMove_Internal(ScheduledLayeredMove);
 	}
+
+#if !UE_BUILD_SHIPPING
+	ENetMode NetMode = GetWorld() ? GetWorld()->GetNetMode() : NM_MAX;
+	UE_LOGF(LogMover, Verbose, "(%ls) UMoverComponent::QueueScheduledLayeredMove: %ls thread queueing a layered move scheduled at [%ls]: %ls.",
+		*ToString(NetMode), IsInGameThread() ? TEXT("Game") : TEXT("Physics"), *ScheduledLayeredMove.SchedulingInfo.ToString(), ScheduledLayeredMove.Move.IsValid() ? *ScheduledLayeredMove.Move->ToSimpleString() : TEXT("INVALID MOVE"));
+#endif // !UE_BUILD_SHIPPING
+}
+
+const TArray<FScheduledLayeredMove>& UMoverComponent::GetQueuedLayeredMoves() const
+{
+	return QueuedLayeredMoves;
+}
+
+void UMoverComponent::ClearQueuedLayeredMoves()
+{
+	QueuedLayeredMoves.Empty();
 }
 
 FMovementModifierHandle UMoverComponent::K2_QueueMovementModifier(const int32& MoveAsRawData)
@@ -1811,53 +1858,57 @@ void UMoverComponent::ScheduleInstantMovementEffect(TSharedPtr<FInstantMovementE
 	FMoverTimeStep TimeStep;
 	if (ensureMsgf(BackendLiaisonComp, TEXT("UMoverComponent::ScheduleInstantMovementEffect was unexpectedly called with a null backend liaison component. The instant movement effect will be ignored.")))
 	{
-		TimeStep.BaseSimTimeMs = BackendLiaisonComp->GetCurrentSimTimeMs();
-		TimeStep.ServerFrame = BackendLiaisonComp->GetCurrentSimFrame();
-		// TimeStep.StepMs is not used by FScheduledInstantMovementEffect::ScheduleEffect
-		QueueInstantMovementEffect(FScheduledInstantMovementEffect::ScheduleEffect(GetWorld(), TimeStep, InstantMovementEffect, /* SchedulingDelaySeconds = */ EventSchedulingMinDelaySeconds));
+		if (ensure(InstantMovementEffect.IsValid()))
+		{
+			TimeStep.BaseSimTimeMs = BackendLiaisonComp->GetCurrentSimTimeMs();
+			TimeStep.ServerFrame = BackendLiaisonComp->GetCurrentSimFrame();
+			// TimeStep.StepMs is not used by UE::Mover::MakeScheduledEffect so we don't set it here
+			QueueScheduledInstantMovementEffect(UE::Mover::MakeScheduledEffect(BackendLiaisonComp, TimeStep, InstantMovementEffect));
+		}
 	}
 }
 
 void UMoverComponent::QueueInstantMovementEffect_Internal(const FMoverTimeStep& TimeStep, TSharedPtr<FInstantMovementEffect> InstantMovementEffect)
 {
-	QueueInstantMovementEffect(FScheduledInstantMovementEffect::ScheduleEffect(GetWorld(), TimeStep, InstantMovementEffect, /* SchedulingDelaySeconds = */ 0.0f));
+	if (ensureMsgf(BackendLiaisonComp, TEXT("UMoverComponent::QueueInstantMovementEffect_Internal was unexpectedly called with a null backend liaison component. The instant movement effect will be ignored.")))
+	{
+		QueueScheduledInstantMovementEffect(UE::Mover::MakeScheduledEffect(BackendLiaisonComp, TimeStep, InstantMovementEffect));
+	}
 }
 
 void UMoverComponent::QueueInstantMovementEffect(TSharedPtr<FInstantMovementEffect> InstantMovementEffect)
 {
 	ensureMsgf(IsInGameThread(), TEXT("UMoverComponent::QueueInstantMovementEffect(TSharedPtr<FInstantMovementEffect>) should only be called from the game thread. Inspect code for incorrect calls."));
 	FMoverTimeStep TimeStep;
-	if (ensureMsgf(BackendLiaisonComp, TEXT("UMoverComponent::ScheduleInstantMovementEffect was unexpectedly called with a null backend liaison component. The instant movement effect will be ignored.")))
+	if (ensureMsgf(BackendLiaisonComp, TEXT("UMoverComponent::QueueInstantMovementEffect was unexpectedly called with a null backend liaison component. The instant movement effect will be ignored.")))
 	{
-		TimeStep.BaseSimTimeMs = BackendLiaisonComp->GetCurrentSimTimeMs();
-		TimeStep.ServerFrame = BackendLiaisonComp->GetCurrentSimFrame();
-		// TimeStep.StepMs is not used by FScheduledInstantMovementEffect::ScheduleEffect
-		QueueInstantMovementEffect(FScheduledInstantMovementEffect::ScheduleEffect(GetWorld(), TimeStep, InstantMovementEffect, /* SchedulingDelaySeconds = */ 0.0f));
+		if (ensure(InstantMovementEffect.IsValid()))
+		{
+			TimeStep.BaseSimTimeMs = BackendLiaisonComp->GetCurrentSimTimeMs();
+			TimeStep.ServerFrame = BackendLiaisonComp->GetCurrentSimFrame();
+			// TimeStep.StepMs is not used by UE::Mover::MakeScheduledEffect so we don't set it here
+			QueueScheduledInstantMovementEffect(UE::Mover::MakeScheduledEffectForNextFrame(BackendLiaisonComp, TimeStep, InstantMovementEffect));
+		}
 	}
 }
 
-void UMoverComponent::QueueInstantMovementEffect(const FScheduledInstantMovementEffect& InstantMovementEffect)
+void UMoverComponent::QueueScheduledInstantMovementEffect(const FScheduledInstantMovementEffect& InstantMovementEffect)
 {
-	if (!ModeFSM)
-	{
-		return;
-	}
-
-	// TODO Move QueueInstantMovementEffect to UMoverSimulation and implement differently in sync or async mode
+	// TODO Move QueueScheduledInstantMovementEffect to UMoverSimulation and implement differently in sync or async mode
 	if (IsInGameThread())
 	{
 		QueuedInstantMovementEffects.Add(InstantMovementEffect);
 	}
-	else
+	else if (ensure(ModeFSM))
 	{
 		ModeFSM->QueueInstantMovementEffect_Internal(InstantMovementEffect);
 	}	
 
-#if !defined(BUILD_SHIPPING) || !BUILD_SHIPPING
+#if !UE_BUILD_SHIPPING
 	ENetMode NetMode = GetWorld() ? GetWorld()->GetNetMode() : NM_MAX;
-	UE_LOGF(LogMover, Verbose, "(%ls) UMoverComponent::QueueInstantMovementEffect: Game Thread queueing an instant movement effect scheduled for frame %d: %ls.",
-		*ToString(NetMode), InstantMovementEffect.ExecutionServerFrame, InstantMovementEffect.Effect.IsValid() ? *InstantMovementEffect.Effect->ToSimpleString() : TEXT("INVALID INSTANT EFFECT"));
-#endif // !defined(BUILD_SHIPPING) || !BUILD_SHIPPING
+	UE_LOGF(LogMover, Verbose, "(%ls) UMoverComponent::QueueScheduledInstantMovementEffect: %ls thread queueing an instant movement effect scheduled at [%ls]: %ls.",
+		*ToString(NetMode), IsInGameThread() ? TEXT("Game") : TEXT("Physics"), *InstantMovementEffect.SchedulingInfo.ToString(), InstantMovementEffect.Effect.IsValid() ? *InstantMovementEffect.Effect->ToSimpleString() : TEXT("INVALID INSTANT EFFECT"));
+#endif // !UE_BUILD_SHIPPING
 }
 
 const TArray<FScheduledInstantMovementEffect>& UMoverComponent::GetQueuedInstantMovementEffects() const
diff --git a/Engine/Plugins/Experimental/Mover/Source/Mover/Private/MoverSchedulingUtils.cpp b/Engine/Plugins/Experimental/Mover/Source/Mover/Private/MoverSchedulingUtils.cpp
new file mode 100644
--- /dev/null
+++ b/Engine/Plugins/Experimental/Mover/Source/Mover/Private/MoverSchedulingUtils.cpp
@@ -0,0 +1,74 @@
+// Copyright Epic Games, Inc. All Rights Reserved.
+
+#include "MoverSchedulingUtils.h"
+#include "Physics/Experimental/PhysScene_Chaos.h"
+#include "PBDRigidsSolver.h"
+
+namespace UE::Mover
+{
+
+FMoverSchedulingInfo MakeSchedulingInfoForScheduledNetworkTime(const TScriptInterface<IMoverBackendLiaisonInterface>& Backend, const FMoverTimeStep& TimeStep)
+{
+	FMoverTime ScheduledTime = ensure(Backend) ? Backend->GetScheduledNetworkTime(FMoverTime(TimeStep.ServerFrame, TimeStep.BaseSimTimeMs)) : FMoverTime();
+	bool bIsFixedDt = Backend ? Backend->IsFixedDt() : false;
+	return FMoverSchedulingInfo(
+								/* ServerIssuanceTime  = */ FMoverTime(TimeStep.ServerFrame, TimeStep.BaseSimTimeMs),
+								/* ServerExecutionTime = */ FMoverTime(ScheduledTime.FrameCount, ScheduledTime.TimeMs),
+								/* bIsFixedDt          = */ bIsFixedDt);
+}
+
+FMoverSchedulingInfo MakeSchedulingInfoForNextFrame(const TScriptInterface<IMoverBackendLiaisonInterface>& Backend, const FMoverTimeStep& TimeStep)
+{
+	bool bIsFixedDt = ensure(Backend) ? Backend->IsFixedDt() : false;
+
+	if (bIsFixedDt)
+	{
+		const int32 CurrentServerFrame = TimeStep.ServerFrame;
+		const int32 ExecutionServerFrame = CurrentServerFrame + 1;
+
+		return FMoverSchedulingInfo(
+				/* ServerIssuanceTime  = */ FMoverTime(CurrentServerFrame, 0.0),
+				/* ServerExecutionTime = */ FMoverTime(ExecutionServerFrame, 0.0),
+				/* bIsFixedDt          = */ bIsFixedDt);
+	}
+	else
+	{
+		const double CurrentServerTimeMs = TimeStep.BaseSimTimeMs;
+		const double ExecutionServerTimeMs = CurrentServerTimeMs + UE_KINDA_SMALL_NUMBER;
+
+		return FMoverSchedulingInfo(
+				/* ServerIssuanceTime  = */ FMoverTime(INDEX_NONE, CurrentServerTimeMs),
+				/* ServerExecutionTime = */ FMoverTime(INDEX_NONE, ExecutionServerTimeMs),
+				/* bIsFixedDt          = */ bIsFixedDt);
+	}
+}
+
+FScheduledInstantMovementEffect MakeScheduledEffect(const TScriptInterface<IMoverBackendLiaisonInterface>& Backend, const FMoverTimeStep& TimeStep, TSharedPtr<FInstantMovementEffect> InstantMovementEffect)
+{
+	ensure(InstantMovementEffect.IsValid());
+
+	return FScheduledInstantMovementEffect(MakeSchedulingInfoForScheduledNetworkTime(Backend, TimeStep), InstantMovementEffect);
+}
+
+FScheduledInstantMovementEffect MakeScheduledEffectForNextFrame(const TScriptInterface<IMoverBackendLiaisonInterface>& Backend, const FMoverTimeStep& TimeStep, TSharedPtr<FInstantMovementEffect> InstantMovementEffect)
+{
+	ensure(InstantMovementEffect.IsValid());
+
+	return FScheduledInstantMovementEffect(MakeSchedulingInfoForNextFrame(Backend, TimeStep), InstantMovementEffect);
+}
+
+FScheduledLayeredMove MakeScheduledMove(const TScriptInterface<IMoverBackendLiaisonInterface>& Backend, const FMoverTimeStep& TimeStep, TSharedPtr<FLayeredMoveBase> Move)
+{
+	ensure(Move.IsValid());
+
+	return FScheduledLayeredMove(MakeSchedulingInfoForScheduledNetworkTime(Backend, TimeStep), Move);
+}
+
+FScheduledLayeredMove MakeScheduledMoveForNextFrame(const TScriptInterface<IMoverBackendLiaisonInterface>& Backend, const FMoverTimeStep& TimeStep, TSharedPtr<FLayeredMoveBase> Move)
+{
+	ensure(Move.IsValid());
+
+	return FScheduledLayeredMove(MakeSchedulingInfoForNextFrame(Backend, TimeStep), Move);
+}
+
+} // namespace UE::Mover
\ No newline at end of file
diff --git a/Engine/Plugins/Experimental/Mover/Source/Mover/Private/MoverSimulationTypes.cpp b/Engine/Plugins/Experimental/Mover/Source/Mover/Private/MoverSimulationTypes.cpp
--- a/Engine/Plugins/Experimental/Mover/Source/Mover/Private/MoverSimulationTypes.cpp
+++ b/Engine/Plugins/Experimental/Mover/Source/Mover/Private/MoverSimulationTypes.cpp
@@ -33,53 +33,6 @@ UScriptStruct* FMoverSimulationEventData::GetScriptStruct() const
 	return FMoverSimulationEventData::StaticStruct();
 }
 
-FScheduledInstantMovementEffect FScheduledInstantMovementEffect::ScheduleEffect(UWorld* World, const FMoverTimeStep& TimeStep, TSharedPtr<FInstantMovementEffect> InstantMovementEffect, float SchedulingDelaySeconds)
-{
-	FPhysScene* Scene = World ? World->GetPhysicsScene() : nullptr;
-	Chaos::FPhysicsSolver* Solver = Scene ? Scene->GetSolver() : nullptr;
-	bool bIsFixedDt = Solver ? Solver->IsUsingFixedDt() : false;
-	if (bIsFixedDt)
-	{
-		int32 ServerFrame = TimeStep.ServerFrame;
-		if (SchedulingDelaySeconds != 0.0f)
-		{
-			const double DeltaTime = Solver->GetAsyncDeltaTime();
-			const int32 DelayFrames = (DeltaTime > UE_SMALL_NUMBER) ? FMath::CeilToInt32(SchedulingDelaySeconds / DeltaTime) : 0;
-			ServerFrame += DelayFrames;
-		}
-		return FScheduledInstantMovementEffect(/* ExecutionServerFrame = */ ServerFrame, /* ExecutionServerTime = */ 0.0, /* bIsFixedDt = */ true, InstantMovementEffect);
-	}
-	else
-	{
-		AGameStateBase* GameState = World ? World->GetGameState() : nullptr;
-		double ServerTime = GameState ? GameState->GetServerWorldTimeSeconds() : 0.0;
-		ServerTime += SchedulingDelaySeconds;
-		
-		return FScheduledInstantMovementEffect(/* ExecutionServerFrame = */ INDEX_NONE, /* ExecutionServerTime = */ ServerTime, /* bIsFixedDt = */ false, InstantMovementEffect);
-	}
-}
-
-FScheduledInstantMovementEffect FScheduledInstantMovementEffect::ScheduleEffectForNextFrame(UWorld* World, const FMoverTimeStep& TimeStep, TSharedPtr<FInstantMovementEffect> InstantMovementEffect)
-{
-	FPhysScene* Scene = World ? World->GetPhysicsScene() : nullptr;
-	Chaos::FPhysicsSolver* Solver = Scene ? Scene->GetSolver() : nullptr;
-	bool bIsFixedDt = Solver ? Solver->IsUsingFixedDt() : false;
-	if (bIsFixedDt)
-	{
-		int32 ServerFrame = TimeStep.ServerFrame + 1;
-
-		return FScheduledInstantMovementEffect(/* ExecutionServerFrame = */ ServerFrame, /* ExecutionServerTime = */ 0.0, /* bIsFixedDt = */ true, InstantMovementEffect);
-	}
-	else
-	{
-		AGameStateBase* GameState = World ? World->GetGameState() : nullptr;
-		double ServerTime = GameState ? GameState->GetServerWorldTimeSeconds() : 0.0;
-		ServerTime += UE_KINDA_SMALL_NUMBER;
-
-		return FScheduledInstantMovementEffect(/* ExecutionServerFrame = */ INDEX_NONE, /* ExecutionServerTime = */ ServerTime, /* bIsFixedDt = */ false, InstantMovementEffect);
-	}
-}
-
 namespace UE::Mover
 {
 	void FSimulationOutputData::Reset()
diff --git a/Engine/Plugins/Experimental/Mover/Source/Mover/Public/Backends/MoverBackendLiaison.h b/Engine/Plugins/Experimental/Mover/Source/Mover/Public/Backends/MoverBackendLiaison.h
--- a/Engine/Plugins/Experimental/Mover/Source/Mover/Public/Backends/MoverBackendLiaison.h
+++ b/Engine/Plugins/Experimental/Mover/Source/Mover/Public/Backends/MoverBackendLiaison.h
@@ -31,11 +31,19 @@ public:
 
 	// Whether this backend will simulate movement asynchronously
 	virtual bool IsAsync() const { return false; }
+	// Whether the underlying simulation advances time in fixed increments
+	virtual bool IsFixedDt() const { return false; }
 	// During a resim, whether this particular mover should resimulate
 	virtual bool ShouldResim() const { return true; }
 
 	// How much delay to apply to scheduled events. This is important for networked events, and should be greater than the RTT to ensure the event will be executed on all end points at the same frame.
 	virtual float GetEventSchedulingMinDelaySeconds() const {return 0.3f;}
+	// Util function that returns a time in the future far enough for all end points to receive a networked message before that time elapses. This is only useful in the context of networked events.
+	virtual FMoverTime GetScheduledNetworkTime(const FMoverTime& Time) const
+	{
+		ensureMsgf(!IsFixedDt(), TEXT("In Fixed Dt it is recommended to also override GetScheduledNetworkTime to return a correct scheduled frame count time"));
+		return FMoverTime(INDEX_NONE, Time.TimeMs + GetEventSchedulingMinDelaySeconds() * 1000);
+	}
 
 	// Pending State: the simulation state currently being authored
 	virtual bool ReadPendingSyncState(OUT FMoverSyncState& OutSyncState) { return false; }
diff --git a/Engine/Plugins/Experimental/Mover/Source/Mover/Public/MovementModeStateMachine.h b/Engine/Plugins/Experimental/Mover/Source/Mover/Public/MovementModeStateMachine.h
--- a/Engine/Plugins/Experimental/Mover/Source/Mover/Public/MovementModeStateMachine.h
+++ b/Engine/Plugins/Experimental/Mover/Source/Mover/Public/MovementModeStateMachine.h
@@ -8,6 +8,7 @@
 #include "MovementMode.h"
 #include "MoveLibrary/RollbackBlackboardWrappers.h"
 #include "InstantMovementEffect.h"
+#include "MoverSimulationTypes.h"
 #include "MovementModeStateMachine.generated.h"
 
 #define UE_API MOVER_API
@@ -58,8 +59,6 @@ public:
 
 	UE_API const UBaseMovementMode* FindMovementMode(FName ModeName) const;
 
-	UE_API void QueueLayeredMove(TSharedPtr<FLayeredMoveBase> Move);
-	
 	UE_API void QueueActiveLayeredMove(const TSharedPtr<FLayeredMoveInstance>& LayeredMove);
 
  	UE_API FMovementModifierHandle QueueMovementModifier(TSharedPtr<FMovementModifierBase> Modifier);
@@ -74,10 +73,17 @@ public:
 
 	// This function is meant to be used only in async mode on the physics thread, not on the game thread
 	void QueueInstantMovementEffect_Internal(const FScheduledInstantMovementEffect& ScheduledInstantMovementEffect);
+
+	// This function is meant to be used only in async mode on the physics thread, not on the game thread
+	void QueueLayeredMove_Internal(const FScheduledLayeredMove& ScheduledLayeredMove);
+
 protected:
 	UE_API void QueueInstantMovementEffect(const FScheduledInstantMovementEffect& ScheduledInstantMovementEffect);
 	UE_API void QueueInstantMovementEffects(const TArray<FScheduledInstantMovementEffect>& ScheduledInstantMovementEffects);
 
+	UE_API void QueueLayeredMove(const FScheduledLayeredMove& ScheduledLayeredMove);
+	UE_API void QueueLayeredMoves(const TArray<FScheduledLayeredMove>& ScheduledLayeredMoves);
+
 	void ProcessEvents(const TArray<TSharedPtr<FMoverSimulationEventData>>& InEvents);
 	UE_API virtual void ProcessSimulationEvent(const FMoverSimulationEventData& EventData);
 
@@ -102,13 +108,13 @@ protected:
 	// Represents the current sim time that's passed, and the next frame number that's next to be simulated.
 	FMoverTimeStep CurrentBaseTimeStep;
 
-	/** Moves that are queued to be added to the simulation at the start of the next sim subtick. Access covered by lock. */
-	TArray<TSharedPtr<FLayeredMoveBase>> QueuedLayeredMoves;
+	/** Moves that are queued to be added to the simulation when their scheduled time has come (earliest is at the start of the next sim subtick). Access covered by lock. */
+	TArray<FScheduledLayeredMove> QueuedLayeredMoves;
 
 	/** Moves that are queued to be added to the simulation at the start of the next sim subtick. Access covered by lock. */
  	TArray<TSharedPtr<FLayeredMoveInstance>> QueuedLayeredMoveInstances;
  	
- 	/** Effects that are queued to be applied to the simulation at the start of the next sim subtick or at the end of this tick.  Access covered by lock. */
+ 	/** Effects that are queued to be applied to the simulation when their scheduled time has come (earliest is at the start of the next sim subtick, or at the end of this tick).  Access covered by lock. */
  	TArray<FScheduledInstantMovementEffect> QueuedInstantEffects;
 
  	/** Modifiers that are queued to be added to the simulation at the start of the next sim subtick. Access covered by lock. */
@@ -134,13 +140,15 @@ private:
 
 	UE_API void ConstructDefaultModes();
 	UE_API void AdvanceToNextMode();
-	UE_API void FlushQueuedMovesToGroup(FLayeredMoveGroup& Group);
+	UE_API void FlushQueuedMovesToGroup(const FMoverTimeStep& TimeStep, FLayeredMoveGroup& Group);
 	// Flushes queued ActiveLayeredMoves to FLayeredMoveInstanceGroup for this frame
  	UE_API void ActivateQueuedMoves(FLayeredMoveInstanceGroup& Group);
  	UE_API void FlushQueuedModifiersToGroup(FMovementModifierGroup& ModifierGroup);
  	UE_API void FlushModifierCancellationsToGroup(FMovementModifierGroup& ActiveModifierGroup);
 	UE_API void FlushTagCancellationsToSyncState(FMoverSyncState& SyncState);
  	UE_API void RollbackModifiers(const FMoverSyncState* InvalidSyncState, const FMoverSyncState* SyncState, const FMoverAuxStateContext* InvalidAuxState, const FMoverAuxStateContext* AuxState);
+	UE_API void RollBackQueuedInstantMovementEffects(const FMoverTimeStep& NewBaseTimeStep);
+	UE_API void RollBackQueuedLayeredMoves(const FMoverTimeStep& NewBaseTimeStep);
 	UE_API bool HasAnyInstantEffectsQueued() const;
  	UE_API bool ApplyInstantEffects(FApplyMovementEffectParams& ApplyEffectParams, FMoverSyncState& OutputState);
 	UE_API AActor* GetOwnerActor() const;
diff --git a/Engine/Plugins/Experimental/Mover/Source/Mover/Public/MoverComponent.h b/Engine/Plugins/Experimental/Mover/Source/Mover/Public/MoverComponent.h
--- a/Engine/Plugins/Experimental/Mover/Source/Mover/Public/MoverComponent.h
+++ b/Engine/Plugins/Experimental/Mover/Source/Mover/Public/MoverComponent.h
@@ -280,6 +280,37 @@ public:
 
 	// Queue a layered move to start during the next simulation frame
 	MOVER_API void QueueLayeredMove(TSharedPtr<FLayeredMoveBase> Move);
+
+	/**
+  	 * Advance Usage. This will only become relevant to Mover in general when Layered Moves are automatically networked (pending, only Chaos Mover networks them by default)
+	 * Queue a scheduled move to take place after delay (tunable in the NetworkPhysicsSettingsComponent)
+	 * ensuring it gets executed on the same frame on all networked end points. @see UNetworkPhysicsSettingsComponent, @see FNetworkPhysicsSettings
+	 * Examples:
+	 * 1. For an Autonomous Proxy Actor: An autonomous proxy might want to schedule a move that can be dodged by other players, for instance and air stomp, so that it starts executing on the same frame for sim proxies
+	 *    giving them a fair chance to dodge it. The server might also need to schedule 
+	 * 2. For an actor moved on the Server authoritatively. Since the server runs behind clients, if it wants to avoid corrections it needs to schedule execution of moves in the future so they get
+	 *    executed on the same frame on all end points. For instance, a platform changing direction as a result of some unpredictable event (such as colliding with a player character) 
+	 *    should be scheduled if we want to guarantee correction-free physical interactions with that platform
+	 * 
+	 * Pros: Everybody executes on the same frame (or at roughly the same time in variable dt), less correction, no surprises.
+     * Cons: Delay in execution, needs to be masked with anticipation anim. Significant lag (several frames).
+	 * 
+	 * Default to using QueueLayeredMove instead, except for advanced use cases where it is important the move get executed on the same frame on all end points.
+	 * 
+	 * @param Move			The layered move to queue, which must be a FLayeredMoveBase sub - type.
+	 */
+	MOVER_API void ScheduleLayeredMove(TSharedPtr<FLayeredMoveBase> Move);
+
+	// Get the queued layered moves. This is mostly for internal use, general users should abstain from calling it.
+	MOVER_API const TArray<FScheduledLayeredMove>& GetQueuedLayeredMoves() const;
+	// Clears the queued layered moves. This is mostly for internal use, general users should abstain from calling it.
+	MOVER_API void ClearQueuedLayeredMoves();
+
+protected:
+	// Queue a scheduled layered move. Thread safe, can be used outside the game thread.
+	void QueueScheduledLayeredMove(const FScheduledLayeredMove& ScheduledLayeredMove);
+
+public:
 	
 	/**
  	 * Queue a Movement Modifier to start during the next simulation frame. This will clone whatever move you pass in, so you'll need to fully set it up before queuing.
@@ -292,7 +323,7 @@ public:
 
 	// Queue a Movement Modifier to start during the next simulation frame.
 	MOVER_API FMovementModifierHandle QueueMovementModifier(TSharedPtr<FMovementModifierBase> Modifier);
-	
+
 	/**
 	 * Cancel any active or queued Modifiers with the handle passed in.
 	 */
@@ -335,8 +366,10 @@ public:
 	 */ 
 	MOVER_API void QueueInstantMovementEffect(TSharedPtr<FInstantMovementEffect> InstantMovementEffect);
 	/** 
+	 * Advance Usage. This will only become relevant to Mover in general when Instant Movement Effects are automatically networked (pending, only Chaos Mover networks them by default)
 	 * Queue a scheduled Instant Movement Effect to take place after delay (tunable in the NetworkPhysicsSettingsComponent)
-	 * ensuring it gets executed on the same frame on all networked end points. @see UNetworkPhysicsSettingsComponent, @see FNetworkPhysicsSettings, @see EventSchedulingMinDelaySeconds
+	 * ensuring it gets executed on the same frame on all networked end points. @see UNetworkPhysicsSettingsComponent, @see FNetworkPhysicsSettings
+	 * Default to using QueueInstantMovementEffect instead, except for advanced use cases where it is important the effect get executed on the same frame on all end points.
 	 * @param InstantMovementEffect			The effect to queue, which must be a FInstantMovementEffect sub - type.
 	 */
 	MOVER_API void ScheduleInstantMovementEffect(TSharedPtr<FInstantMovementEffect> InstantMovementEffect);
@@ -350,7 +383,7 @@ public:
 
 protected:
 	// Queue a scheduled instant movement effect. Thread safe, can be used outside the game thread.
-	void QueueInstantMovementEffect(const FScheduledInstantMovementEffect& ScheduledInstantMovementEffect);
+	void QueueScheduledInstantMovementEffect(const FScheduledInstantMovementEffect& ScheduledInstantMovementEffect);
 
 public:	
 	// Queue a movement mode change to occur during the next simulation frame. If bShouldReenter is true, then a mode change will occur even if already in that mode.
@@ -829,9 +862,12 @@ private:
 	UPROPERTY(EditDefaultsOnly, Category = "Mover|Constraints")
 	FPlanarConstraint PlanarConstraint;
 
-	/** Effects queued to be applied to the simulation at a given frame. If the frame happens to be in the past, the effect will be applied at the earliest occasion */
+	/** Effects queued to be applied to the simulation at a given frame/time. If the frame happens to be in the past, the effect will be applied at the earliest occasion */
 	TArray<FScheduledInstantMovementEffect> QueuedInstantMovementEffects;
 
+	/** Layered moves queued to be applied to the simulation at a given frame/time. If the frame happens to be in the past, the move will be applied at the earliest occasion */
+	TArray<FScheduledLayeredMove> QueuedLayeredMoves;
+
 public:
 
 	// If enabled, the movement of the primary visual component will be smoothed via an offset from the root moving component. This is useful in fixed-tick simulations with variable rendering rates.
@@ -913,10 +949,6 @@ protected:
 	UPROPERTY(EditDefaultsOnly, Category = "Mover")
 	bool bSupportsKinematicBasedMovement = true;
 
-	// Delay added to scheduled instant movement effects
-	// This value is cached from the settings found in the network settings component
-	float EventSchedulingMinDelaySeconds = 0.3f;
-
 	FMoverAuxStateContext CachedLastAuxState;
 
 	friend class UBaseMovementMode;
diff --git a/Engine/Plugins/Experimental/Mover/Source/Mover/Public/MoverSchedulingUtils.h b/Engine/Plugins/Experimental/Mover/Source/Mover/Public/MoverSchedulingUtils.h
new file mode 100644
--- /dev/null
+++ b/Engine/Plugins/Experimental/Mover/Source/Mover/Public/MoverSchedulingUtils.h
@@ -0,0 +1,33 @@
+// Copyright Epic Games, Inc. All Rights Reserved.
+
+#pragma once
+
+#include "MoverSimulationTypes.h"
+#include "Backends/MoverBackendLiaison.h"
+
+namespace UE
+{
+namespace Mover
+{
+	/** Turns a FInstantMovementEffect into a scheduled one (FScheduledInstantMovementEffect)
+	*	The effect can be scheduled to apply immediately, or scheduled to apply with a delay
+	*   @param Backend the simulation backend which implements the details of whether the simulation advances in fixed dt and how long to schedule in the future for networked events
+	*   @param TimeStep the time step of the current or upcoming tick
+	*   @param InstantMovementEffect the effect to schedule
+	*/
+	MOVER_API FScheduledInstantMovementEffect MakeScheduledEffect(const TScriptInterface<IMoverBackendLiaisonInterface>& Backend, const FMoverTimeStep& TimeStep, TSharedPtr<FInstantMovementEffect> InstantMovementEffect);
+
+	// Same as above but schedules the effect to run on the next frame. This is best used for fixed tick mode, it will add a small SchedulingDelaySeconds in the variable tick case
+	MOVER_API FScheduledInstantMovementEffect MakeScheduledEffectForNextFrame(const TScriptInterface<IMoverBackendLiaisonInterface>& Backend, const FMoverTimeStep& TimeStep, TSharedPtr<FInstantMovementEffect> InstantMovementEffect);
+
+	/** Turns a FLayeredMove into a scheduled one (FScheduledLayeredMove)
+	*	The move can be scheduled to apply immediately, or scheduled to apply with a delay
+	*   @param Backend the simulation backend which implements the details of whether the simulation advances in fixed dt and how long to schedule in the future for networked events
+	*   @param TimeStep the time step of the current or upcoming tick
+	*   @param LayeredMove the move to schedule
+	*/
+	MOVER_API FScheduledLayeredMove MakeScheduledMove(const TScriptInterface<IMoverBackendLiaisonInterface>& Backend, const FMoverTimeStep& TimeStep, TSharedPtr<FLayeredMoveBase> LayeredMove);
+	// Same as above but schedules the effect to run on the next frame. This is best used for fixed tick mode, it will add a small SchedulingDelaySeconds in the variable tick case
+	MOVER_API FScheduledLayeredMove MakeScheduledMoveForNextFrame(const TScriptInterface<IMoverBackendLiaisonInterface>& Backend, const FMoverTimeStep& TimeStep, TSharedPtr<FLayeredMoveBase> LayeredMove);
+}
+}
diff --git a/Engine/Plugins/Experimental/Mover/Source/Mover/Public/MoverSimulationTypes.h b/Engine/Plugins/Experimental/Mover/Source/Mover/Public/MoverSimulationTypes.h
--- a/Engine/Plugins/Experimental/Mover/Source/Mover/Public/MoverSimulationTypes.h
+++ b/Engine/Plugins/Experimental/Mover/Source/Mover/Public/MoverSimulationTypes.h
@@ -69,83 +69,112 @@ struct FMovementModeTickEndState
 	// Affirms that no state changes were made during this simulation tick. Can help optimizations if modes set this during sim tick.
 	UPROPERTY(BlueprintReadWrite, Category=Mover)
 	bool bEndedWithNoChanges = false;
-
 };
 
 USTRUCT()
-struct FScheduledInstantMovementEffect
+struct FMoverSchedulingInfo
 {
 	GENERATED_BODY()
 
-	/** Turns a FInstantMovementEffect into a scheduled one (FScheduledInstantMovementEffect)
-	*	The effect can be scheduled to apply immediately, or scheduled to apply with a delay
-	*   This function should not be called on the game thread
-	*   @param World The world, used to retrieve the current server frame in async mode, or the sim time otherwise
-	*   @param TimeStep the time step of the current or upcoming tick
-	*   @param InstantMovementEffect the effect to schedule
-	*   @param SchedulingDelaySeconds Scheduling delay to ensure it applies on all end points on the same frame (this is only perfectly accurate when simulation dt is fixed)
-	*/
-	MOVER_API static FScheduledInstantMovementEffect ScheduleEffect(UWorld* World, const FMoverTimeStep& TimeStep, TSharedPtr<FInstantMovementEffect> InstantMovementEffect, float SchedulingDelaySeconds = 0.0f);
-	// Same as above but schedules the effect to run on the next frame. This is best used for fixed tick mode, it will add a small SchedulingDelaySeconds in the variable tick case
-	MOVER_API static FScheduledInstantMovementEffect ScheduleEffectForNextFrame(UWorld* World, const FMoverTimeStep& TimeStep, TSharedPtr<FInstantMovementEffect> InstantMovementEffect);
+	bool ShouldExecuteAtTimeStep(const FMoverTimeStep& TimeStep) const
+	{
+		return bIsFixedDt ? ShouldExecuteAtFrame(TimeStep.ServerFrame) : ShouldExecuteAtTimeMs(TimeStep.BaseSimTimeMs);
+	}
+
+	bool ShouldExecuteAtTime(const FMoverTime& ServerTime) const
+	{
+		return bIsFixedDt ? ShouldExecuteAtFrame(ServerTime.FrameCount) : ShouldExecuteAtTimeMs(ServerTime.TimeMs);
+	}
 
 	bool ShouldExecuteAtFrame(int32 CurrentServerFrame) const
 	{
-		ensureMsgf(bIsFixedDt, TEXT("In variable delta time mode, use the version of ShouldExecute that takes a floating point time"));
-		return (CurrentServerFrame >= ExecutionServerFrame);
+		ensureMsgf(bIsFixedDt, TEXT("In variable delta time mode, use ShouldExecuteAtTimeMs"));
+		return (ensure(ServerExecutionTime.FrameCount != INDEX_NONE) && CurrentServerFrame >= ServerExecutionTime.FrameCount);
 	}
 
-	bool ShouldExecuteAtTime(double CurrentServerTime) const
+	bool ShouldExecuteAtTimeMs(double CurrentServerTimeMs) const
 	{
-		ensureMsgf(!bIsFixedDt, TEXT("In fixed delta time mode, use the version of ShouldExecute that takes a frame number"));
-		return (CurrentServerTime >= ExecutionServerTimeSeconds);
+		ensureMsgf(!bIsFixedDt, TEXT("In fixed delta time mode, use ShouldExecuteAtFrame"));
+		return (CurrentServerTimeMs >= ServerExecutionTime.TimeMs);
 	}
 
-	void NetSerialize(const FNetSerializeParams& P)
+	bool ShouldRollBackAtTimeStep(const FMoverTimeStep& NewTimeStep) const
 	{
-		P.Ar.SerializeBits(&bIsFixedDt, 1);
-		if (bIsFixedDt)
-		{
-			P.Ar << ExecutionServerFrame;
-		}
-		else
-		{
-			P.Ar << ExecutionServerTimeSeconds;
-		}
-		
-		Effect->NetSerialize(P.Ar);
+		return bIsFixedDt ? ShouldRollBackAtFrame(NewTimeStep.ServerFrame) : ShouldRollBackAtTimeMs(NewTimeStep.BaseSimTimeMs);
 	}
 
-	void ToString(FAnsiStringBuilderBase& Out) const
+	bool ShouldRollBackAtFrame(int32 NewServerFrame) const
 	{
-		if (bIsFixedDt)
-		{
-			Out.Appendf("ExecutionServerFrame: %d", ExecutionServerFrame);
-		}
-		else
-		{
-			Out.Appendf("ExecutionDateSeconds: %f", ExecutionServerTimeSeconds);
-		}
+		ensureMsgf(bIsFixedDt, TEXT("In variable delta time mode, use ShouldRollBackAtTimeMs"));
+		return (ensure(NewServerFrame != INDEX_NONE) && ServerIssuanceTime.FrameCount >= NewServerFrame);
+	}
 
-		Out.Appendf(" | Effect = % s", *(Effect.IsValid() ? Effect->ToSimpleString() : "Invalid"));
+	bool ShouldRollBackAtTimeMs(double NewServerTimeMs) const
+	{
+		ensureMsgf(!bIsFixedDt, TEXT("In fixed delta time mode, use ShouldRollBackAtFrame"));
+		return (ServerIssuanceTime.TimeMs >= NewServerTimeMs);
 	}
 
-	// Server frame at which this instant movement effect should be applied
-	// Only valid if bIsFixedDt is true, i.e. in fixed time step mode
-	UPROPERTY(VisibleAnywhere, Category = "Mover")
-	int32 ExecutionServerFrame = INDEX_NONE;
+	FString ToString() const
+	{
+		return FString::Printf(TEXT("Issuance: %s | Execution: %s | bIsFixedDt = %d "), *ServerIssuanceTime.ToString(), *ServerExecutionTime.ToString(), bIsFixedDt ? 1 : 0);
+	}
 
-	// Server Time (in seconds) after which this instant movement effect should be applied
-	// Only valid if bIsFixedDt is false, i.e. in variable time step mode
+	// Server time at which this scheduled event was issued
+	// bIsFixedDt TRUE: 
+	//	   - ServerIssuanceTime.FrameCount will be valid and set to the backend's GetCurrentSimFrame, i.e. that of the server frame count of the first async step that will be processing this.
+	//     - ServerIssuanceTime.TimeMs MAY be set to ServerIssuanceTime.TimeMs * AsyncFixedDt * 1000, is only here for reference, and it NOT an interpolated time
+	// bIsFixedDt FALSE:
+	//     - ServerIssuanceTime.FrameCount will be invalid, do not use
+	//     - ServerIssuanceTime.TimeMs will be set to the backend's GetCurrentSimTimeMs
 	UPROPERTY(VisibleAnywhere, Category = "Mover")
-	double ExecutionServerTimeSeconds = 0.0;
+	FMoverTime ServerIssuanceTime;
+
+	// Server time at which this scheduled event should take effect
+	// bIsFixedDt TRUE: 
+	//	   - ServerExecutionTime.FrameCount will be valid and set relative to the backend's GetCurrentSimFrame, i.e. that of the server frame count of the first async step that will be processing this.
+	//     - ServerExecutionTime.TimeMs MAY be set to ServerExecutionTime.FrameCount * AsyncFixedDt * 1000, is only here for reference, and it NOT an interpolated time
+	// bIsFixedDt FALSE:
+	//     - ServerExecutionTime.FrameCount will be invalid, do not use
+	//     - ServerExecutionTime.TimeMs will be set relative to the backend's GetCurrentSimTimeMs
+	UPROPERTY(VisibleAnywhere, Category = "Mover")
+	FMoverTime ServerExecutionTime;
 
+	// Whether this time was set in the context of a fixed or variable time advance scheme. This should correspond to the backend's IsFixedDt()
 	UPROPERTY(VisibleAnywhere, Category = "Mover")
 	bool bIsFixedDt = true;
+};
+
+USTRUCT()
+struct FScheduledInstantMovementEffect
+{
+	GENERATED_BODY()
+
+	void ToString(FAnsiStringBuilderBase& Out) const
+	{
+		Out.Appendf("Scheduling = %s | Effect = %s", TCHAR_TO_ANSI(*SchedulingInfo.ToString()), Effect.IsValid() ? TCHAR_TO_ANSI(*Effect->ToSimpleString()) : "Invalid");
+	}
+
+	FMoverSchedulingInfo SchedulingInfo;
 
 	TSharedPtr<FInstantMovementEffect> Effect;
 };
 
+USTRUCT()
+struct FScheduledLayeredMove
+{
+	GENERATED_BODY()
+
+	void ToString(FAnsiStringBuilderBase& Out) const
+	{
+		Out.Appendf("Scheduling = %s | Move = %s", TCHAR_TO_ANSI(*SchedulingInfo.ToString()), Move.IsValid() ? TCHAR_TO_ANSI(*Move->ToSimpleString()) : "Invalid");
+	}
+
+	FMoverSchedulingInfo SchedulingInfo;
+
+	TSharedPtr<FLayeredMoveBase> Move;
+};
+
 /**
  * The client generates this representation of "input" to the simulated actor for one simulation frame. This can be direct mapping
  * of controls, or more abstract data. It is composed of a collection of typed structs that can be customized per project.
diff --git a/Engine/Plugins/Experimental/Mover/Source/Mover/Public/MoverTypes.h b/Engine/Plugins/Experimental/Mover/Source/Mover/Public/MoverTypes.h
--- a/Engine/Plugins/Experimental/Mover/Source/Mover/Public/MoverTypes.h
+++ b/Engine/Plugins/Experimental/Mover/Source/Mover/Public/MoverTypes.h
@@ -96,7 +96,7 @@ struct FMoverEventContext
 
 	// Server frame that this event occurred
 	UPROPERTY(EditAnywhere, BlueprintReadWrite, Category = Mover)
-	int32 ServerFrame = -1;
+	int32 ServerFrame = INDEX_NONE;
 
 	// Did this event occur while resimulating a past sim frame?
 	UPROPERTY(EditAnywhere, BlueprintReadWrite, Category = Mover)
@@ -107,6 +107,29 @@ struct FMoverEventContext
 	uint8 bIsCausedByRollback : 1 = 0;
 };
 
+// A time that can be a frame count or a floating point value.
+// When time advances in fixed time steps a frame count will be provided to allow for exact time arithmetic
+// A floating point time will also be provided in milliseconds
+// When time advances in variable time steps, the frame count will be INDEX_NONE (-1) and only the floating point time will be valid
+USTRUCT(BlueprintType)
+struct FMoverTime
+{
+	GENERATED_BODY()
+
+	FString ToString() const
+	{
+		return (FrameCount != INDEX_NONE) ? FString::Printf(TEXT("FrameCount = %d, TimeMs = %.3f"), FrameCount, TimeMs) : FString::Printf(TEXT("TimeMs = %.3f"), TimeMs);
+	}
+	
+	// Number of fixed time steps that have occurred since time 0. Only valid when time advances in fixed time steps.
+	UPROPERTY(EditAnywhere, BlueprintReadWrite, Category=Mover)
+	int32 FrameCount=INDEX_NONE;
+
+	// Time in milliseconds. In fixed time steps this should always be equal to FrameCount * Duration of a Fixed Time Step
+	UPROPERTY(EditAnywhere, BlueprintReadWrite, Category=Mover)
+	double TimeMs=-1.0;
+};
+
 
 USTRUCT(BlueprintType)
 struct FMoverTimeStep
@@ -116,7 +139,7 @@ struct FMoverTimeStep
 
 	// The server simulation frame this timestep is associated with. This is the frame that's being worked on now.
 	UPROPERTY(EditAnywhere, BlueprintReadWrite, Category=Mover)
-	int32 ServerFrame=-1;
+	int32 ServerFrame=INDEX_NONE;
 
 	// Starting simulation time (in server simulation timespace)
 	UPROPERTY(EditAnywhere, BlueprintReadWrite, Category=Mover)
