[Chaos Mover] Move input production to async callback InjectInputs_External instead of current rewind callback. This enables us to have a single entry point for both networked and non-networked games

#rnx

#Synced-CL 51708584

diff --git a/Engine/Plugins/Experimental/ChaosMover/Source/ChaosMover/Private/Backends/ChaosMoverAsyncCallback.cpp b/Engine/Plugins/Experimental/ChaosMover/Source/ChaosMover/Private/Backends/ChaosMoverAsyncCallback.cpp
--- a/Engine/Plugins/Experimental/ChaosMover/Source/ChaosMover/Private/Backends/ChaosMoverAsyncCallback.cpp
+++ b/Engine/Plugins/Experimental/ChaosMover/Source/ChaosMover/Private/Backends/ChaosMoverAsyncCallback.cpp
@@ -8,13 +8,41 @@
 #include "ChaosMover/ChaosMoverConsoleVariables.h"
 #include "ChaosMover/ChaosMoverSimulation.h"
 #include "ChaosMover/ChaosMoverSimulationTypes.h"
+#include "ChaosMoverSubsystem.h"
 #include "PBDRigidsSolver.h"
 
 namespace UE::ChaosMover
 {
-	void FAsyncCallback::ProcessInputs_Internal(int32 PhysicsStep)
+	void FAsyncCallback::OnInjectInputs_External(int32 PhysicsStep, int32 NumSteps)
 	{
-		QUICK_SCOPE_CYCLE_COUNTER(ChaosMover_AsyncCallback_ProcessInputs_Internal);
+		QUICK_SCOPE_CYCLE_COUNTER(ChaosMover_AsyncCallback_OnInjectInputs_External);
+
+		if (Subsystem)
+		{
+			TArray<TWeakObjectPtr<UChaosMoverBackendComponent>>& Backends = Subsystem->ValidateAndGetBackends();
+			const FMoverTimeStep& TimeStep = Subsystem->GetMoverTimeStep();
+			const int32 NetPhysicsTickOffset = Subsystem->GetNetworkPhysicsTickOffset();
+
+			UE::ChaosMover::FAsyncCallbackInput* AsyncInput = GetProducerInputData_External();
+			AsyncInput->Reset();
+			AsyncInput->InputData.SetNum(Backends.Num());
+			AsyncInput->PhysicsSolver = static_cast<Chaos::FPhysicsSolver*>(GetSolver());
+			AsyncInput->NetworkPhysicsTickOffset = NetPhysicsTickOffset;
+			for (int Idx = 0; Idx < Backends.Num(); ++Idx)
+			{
+				TWeakObjectPtr<UChaosMoverBackendComponent>& Backend = Backends[Idx];
+				AsyncInput->Backends.Add(Backend);
+				if (Backend.Get())
+				{
+					Backend->ProduceInputData(PhysicsStep, NumSteps, TimeStep, AsyncInput->InputData[Idx]);
+				}
+			}
+		}
+	}
+
+	void FAsyncCallback::OnProcessInputs_Internal(int32 PhysicsStep)
+	{
+		QUICK_SCOPE_CYCLE_COUNTER(ChaosMover_AsyncCallback_OnProcessInputs_Internal);
 
 		const FAsyncCallbackInput* AsyncInput = GetConsumerInput_Internal();
 
diff --git a/Engine/Plugins/Experimental/ChaosMover/Source/ChaosMover/Private/Backends/ChaosMoverAsyncCallback.h b/Engine/Plugins/Experimental/ChaosMover/Source/ChaosMover/Private/Backends/ChaosMoverAsyncCallback.h
--- a/Engine/Plugins/Experimental/ChaosMover/Source/ChaosMover/Private/Backends/ChaosMoverAsyncCallback.h
+++ b/Engine/Plugins/Experimental/ChaosMover/Source/ChaosMover/Private/Backends/ChaosMoverAsyncCallback.h
@@ -6,6 +6,7 @@
 #include "ChaosMover/ChaosMoverSimulationTypes.h"
 
 class UChaosMoverBackendComponent;
+class UChaosMoverSubsystem;
 
 namespace UE::ChaosMover
 {
@@ -41,13 +42,22 @@ namespace UE::ChaosMover
 	class FAsyncCallback : public Chaos::TSimCallbackObject<
 		FAsyncCallbackInput,
 		FAsyncCallbackOutput,
+		Chaos::ESimCallbackOptions::InjectInputsExternal |
+		Chaos::ESimCallbackOptions::ProcessInputsInternal |
 		Chaos::ESimCallbackOptions::Presimulate |
 		Chaos::ESimCallbackOptions::ContactModification |
 		Chaos::ESimCallbackOptions::PostSolve |
 		Chaos::ESimCallbackOptions::Rewind>
 	{
+	public:
+		void SetSubsystem(UChaosMoverSubsystem* InSubsystem)
+		{
+			Subsystem = InSubsystem;
+		}
+
 	protected:
-		virtual void ProcessInputs_Internal(int32 PhysicsStep) override;
+		virtual void OnInjectInputs_External(int32 PhysicsStep, int32 NumSteps) override;
+		virtual void OnProcessInputs_Internal(int32 PhysicsStep) override;
 		virtual void OnPreSimulate_Internal() override;
 		virtual void OnContactModification_Internal(Chaos::FCollisionContactModifier& Modifier) override;
 		virtual void OnPostSolve_Internal() override;
@@ -59,7 +69,9 @@ namespace UE::ChaosMover
 		const TArray<int32, TInlineAllocator<64>>* DetermineResimulatedBackendIndices(const TArray<TWeakObjectPtr<UChaosMoverBackendComponent>>& InBackends);
 
 		// Array of resimulated backend indices. This array is used only in the resim case.
-		// During a resim frame, it is reset and filled in ProcessInputs_Internal and reused in subsequent callbacks on that same frame, such as OnPreSimulate_Internal
+		// During a resim frame, it is reset and filled in OnProcessInputs_Internal and reused in subsequent callbacks on that same frame, such as OnPreSimulate_Internal
 		TArray<int32, TInlineAllocator<64>> ResimulatedBackendIndices;
+
+		UChaosMoverSubsystem* Subsystem = nullptr;
 	};
 }
\ No newline at end of file
diff --git a/Engine/Plugins/Experimental/ChaosMover/Source/ChaosMover/Private/Backends/ChaosMoverSubsystem.cpp b/Engine/Plugins/Experimental/ChaosMover/Source/ChaosMover/Private/Backends/ChaosMoverSubsystem.cpp
--- a/Engine/Plugins/Experimental/ChaosMover/Source/ChaosMover/Private/Backends/ChaosMoverSubsystem.cpp
+++ b/Engine/Plugins/Experimental/ChaosMover/Source/ChaosMover/Private/Backends/ChaosMoverSubsystem.cpp
@@ -31,11 +31,7 @@ void UChaosMoverSubsystem::OnWorldBeginPlay(UWorld& InWorld)
 		if (Chaos::FPhysicsSolver* Solver = PhysScene->GetSolver())
 		{
 			AsyncCallback = Solver->CreateAndRegisterSimCallbackObject_External<UE::ChaosMover::FAsyncCallback>();
-
-			if (FNetworkPhysicsCallback* SolverCallback = static_cast<FNetworkPhysicsCallback*>(Solver->GetRewindCallback()))
-			{
-				InjectInputsExternalCallbackHandle = SolverCallback->InjectInputsExternal.AddUObject(this, &UChaosMoverSubsystem::InjectInputs_External);
-			}
+			AsyncCallback->SetSubsystem(this);
 		}
 	}
 }
@@ -46,23 +42,16 @@ void UChaosMoverSubsystem::Deinitialize()
 	{
 		if (FPhysScene* PhysScene = World->GetPhysicsScene())
 		{
-			if (Chaos::FPhysicsSolver* Solver = PhysScene->GetSolver())
+			if (PhysScenePostTickCallbackHandle.IsValid())
 			{
-				if (InjectInputsExternalCallbackHandle.IsValid())
-				{
-					if (FNetworkPhysicsCallback* SolverCallback = static_cast<FNetworkPhysicsCallback*>(Solver->GetRewindCallback()))
-					{
-						SolverCallback->InjectInputsExternal.Remove(InjectInputsExternalCallbackHandle);
-					}
-				}
-
-				if (PhysScenePostTickCallbackHandle.IsValid())
-				{
-					PhysScene->OnPhysScenePostTick.Remove(PhysScenePostTickCallbackHandle);
-				}
+				PhysScene->OnPhysScenePostTick.Remove(PhysScenePostTickCallbackHandle);
+			}
 
+			if (Chaos::FPhysicsSolver* Solver = PhysScene->GetSolver())
+			{
 				if (AsyncCallback)
 				{
+					AsyncCallback->SetSubsystem(nullptr);
 					Solver->UnregisterAndFreeSimCallbackObject_External(AsyncCallback);
 					AsyncCallback = nullptr;
 				}
@@ -94,14 +83,14 @@ void UChaosMoverSubsystem::Register(TWeakObjectPtr<UChaosMoverBackendComponent>
 						UE_DEBUG_BREAK();
 #endif
 					}
-				}
 
-				if (!bProjectSettingPhysicsPredictionEnabled)
-				{
-					UE_LOG(LogChaosMover, Warning, TEXT("ChaosMoverSubsystem: bEnablePhysicsPrediction project setting is disabled. Enable 'Enable Physics Prediction' in Project Settings > Engine > Physics."));
+					if (!bProjectSettingPhysicsPredictionEnabled)
+					{
+						UE_LOG(LogChaosMover, Warning, TEXT("ChaosMoverSubsystem: bEnablePhysicsPrediction project setting is disabled. Enable 'Enable Physics Prediction' in Project Settings > Engine > Physics."));
 #if !UE_BUILD_SHIPPING
-					UE_DEBUG_BREAK();
+						UE_DEBUG_BREAK();
 #endif
+					}
 				}
 			}
 		}
@@ -115,76 +104,14 @@ void UChaosMoverSubsystem::Unregister(TWeakObjectPtr<UChaosMoverBackendComponent
 	Backends.Remove(InBackend);
 }
 
-void UChaosMoverSubsystem::InjectInputs_External(int32 PhysicsStep, int32 NumSteps)
+TArray<TWeakObjectPtr<UChaosMoverBackendComponent>>& UChaosMoverSubsystem::ValidateAndGetBackends()
 {
-	if (!AsyncCallback)
-	{
-		return;
-	}
-
-	QUICK_SCOPE_CYCLE_COUNTER(ChaosMover_InjectInputs_External);
-
 	// Clear invalid data before we attempt to use it.
 	Backends.RemoveAllSwap([](const TWeakObjectPtr<UChaosMoverBackendComponent> Backend) {
 		return !Backend.IsValid();
 		});
 
-	UE::ChaosMover::FAsyncCallbackInput* AsyncInput = AsyncCallback->GetProducerInputData_External();
-	AsyncInput->Reset();
-	AsyncInput->InputData.SetNum(Backends.Num());
-
-	for (TWeakObjectPtr<UChaosMoverBackendComponent>& Backend : Backends)
-	{
-		AsyncInput->Backends.Add(Backend);
-	}
-
-	// Refresh the network physics tick offset in the async input
-	const int32 NetPhysicsTickOffset = GetNetworkPhysicsTickOffset();
-	AsyncInput->NetworkPhysicsTickOffset = NetPhysicsTickOffset;
-
-	// Compute the time step
-	FMoverTimeStep TimeStep;
-	if (UWorld* World = GetWorld())
-	{
-		if (FPhysScene* Scene = World->GetPhysicsScene())
-		{
-			if (Chaos::FPhysicsSolver* Solver = Scene->GetSolver())
-			{
-				AsyncInput->PhysicsSolver = Solver;
-
-				TimeStep.BaseSimTimeMs = Solver->GetPhysicsResultsTime_External() * 1000.0;
-				
-				if (Solver->IsUsingAsyncResults())
-				{
-					const float AsyncDt = Solver->GetAsyncDeltaTime();
-					if (AsyncDt > UE_SMALL_NUMBER)
-					{
-						TimeStep.ServerFrame = FMath::FloorToInt(Solver->GetPhysicsResultsTime_External() / AsyncDt);
-					}
-					else
-					{
-						TimeStep.ServerFrame = Solver->GetCurrentFrame();
-					}
-					
-					TimeStep.StepMs = AsyncDt * 1000.0f;
-				}
-				else
-				{
-					TimeStep.ServerFrame = Solver->GetCurrentFrame();
-					TimeStep.StepMs = FMath::Clamp(World->GetDeltaSeconds(), Solver->GetMinDeltaTime_External(), Solver->GetMaxDeltaTime_External()) * 1000.0f;
-				}
-
-				TimeStep.ServerFrame += NetPhysicsTickOffset;
-			}
-		}
-	}
-
-	auto LambdaParallelUpdate = [PhysicsStep, NumSteps, AsyncInput, &TimeStep, this](int32 Idx) {
-		Backends[Idx]->ProduceInputData(PhysicsStep, NumSteps, TimeStep, AsyncInput->InputData[Idx]);
-	};
-
-	const bool ForceSingleThread = UE::ChaosMover::CVars::bForceSingleThreadedGT;
-	Chaos::PhysicsParallelFor(Backends.Num(), LambdaParallelUpdate, ForceSingleThread);
+	return Backends;
 }
 
 void UChaosMoverSubsystem::OnPostPhysicsTick(FChaosScene* Scene)
@@ -246,3 +173,42 @@ int32 UChaosMoverSubsystem::GetNetworkPhysicsTickOffset() const
 
 	return 0;
 }
+
+FMoverTimeStep UChaosMoverSubsystem::GetMoverTimeStep() const
+{
+	FMoverTimeStep TimeStep;
+	if (UWorld* World = GetWorld())
+	{
+		if (FPhysScene* Scene = World->GetPhysicsScene())
+		{
+			if (Chaos::FPhysicsSolver* Solver = Scene->GetSolver())
+			{
+				TimeStep.BaseSimTimeMs = Solver->GetPhysicsResultsTime_External() * 1000.0;
+
+				if (Solver->IsUsingAsyncResults())
+				{
+					const float AsyncDt = Solver->GetAsyncDeltaTime();
+					if (AsyncDt > UE_SMALL_NUMBER)
+					{
+						TimeStep.ServerFrame = FMath::FloorToInt(Solver->GetPhysicsResultsTime_External() / AsyncDt);
+					}
+					else
+					{
+						TimeStep.ServerFrame = Solver->GetCurrentFrame();
+					}
+
+					TimeStep.StepMs = AsyncDt * 1000.0f;
+				}
+				else
+				{
+					TimeStep.ServerFrame = Solver->GetCurrentFrame();
+					TimeStep.StepMs = FMath::Clamp(World->GetDeltaSeconds(), Solver->GetMinDeltaTime_External(), Solver->GetMaxDeltaTime_External()) * 1000.0f;
+				}
+
+				TimeStep.ServerFrame += GetNetworkPhysicsTickOffset();;
+			}
+		}
+	}
+
+	return TimeStep;
+}
diff --git a/Engine/Plugins/Experimental/ChaosMover/Source/ChaosMover/Private/Backends/ChaosMoverSubsystem.h b/Engine/Plugins/Experimental/ChaosMover/Source/ChaosMover/Private/Backends/ChaosMoverSubsystem.h
--- a/Engine/Plugins/Experimental/ChaosMover/Source/ChaosMover/Private/Backends/ChaosMoverSubsystem.h
+++ b/Engine/Plugins/Experimental/ChaosMover/Source/ChaosMover/Private/Backends/ChaosMoverSubsystem.h
@@ -15,6 +15,7 @@ namespace UE::ChaosMover
 
 class FChaosScene;
 class UChaosMoverBackendComponent;
+struct FMoverTimeStep;
 
 UCLASS()
 class UChaosMoverSubsystem : public UWorldSubsystem
@@ -25,14 +26,17 @@ public:
 	void Register(TWeakObjectPtr<UChaosMoverBackendComponent> Backend);
 	void Unregister(TWeakObjectPtr<UChaosMoverBackendComponent> Backend);
 
+	// Remove invalid backends before returning
+	TArray<TWeakObjectPtr<UChaosMoverBackendComponent>>& ValidateAndGetBackends();
+
+	FMoverTimeStep GetMoverTimeStep() const;
+	int32 GetNetworkPhysicsTickOffset() const;
+
 protected:
 	virtual bool DoesSupportWorldType(const EWorldType::Type WorldType) const override;
 	virtual void OnWorldBeginPlay(UWorld& InWorld) override;
 	virtual void Deinitialize() override;
 
-	int32 GetNetworkPhysicsTickOffset() const;
-
-	void InjectInputs_External(int32 PhysicsStep, int32 NumSteps);
 	void OnPostPhysicsTick(FChaosScene* Scene);
 
 	TArray<TWeakObjectPtr<UChaosMoverBackendComponent>> Backends;
