[ChaosMover/GASP] Trajectory Prediction support for ChaosMover

* When enabled, ChaosMoverSimulation calculates predicted trajectory in OnSimTick
* This is plumbed out as part of AdditionalOutputData, cached on ChaosCharacterMoverComponent and consumed by *TrajectoryPredictor::Predict()
* Simple backup of the state for now, can be worked on to be more elaborate/generic later

#Synced-CL 51798741

diff --git a/Engine/Plugins/Experimental/ChaosMover/Source/ChaosMover/Private/Backends/ChaosMoverBackend.cpp b/Engine/Plugins/Experimental/ChaosMover/Source/ChaosMover/Private/Backends/ChaosMoverBackend.cpp
--- a/Engine/Plugins/Experimental/ChaosMover/Source/ChaosMover/Private/Backends/ChaosMoverBackend.cpp
+++ b/Engine/Plugins/Experimental/ChaosMover/Source/ChaosMover/Private/Backends/ChaosMoverBackend.cpp
@@ -568,6 +568,12 @@ void UChaosMoverBackendComponent::ProduceInputData(int32 PhysicsStep, int32 NumS
 	InjectInstantMovementEffectsIntoSim(InputData);
 	InjectLayeredMovesIntoSim(InputData);
 
+	// TODO @Harsha consider moving ProduceLocalInput() to the base mover component class
+	if (const UChaosCharacterMoverComponent* ChaosCharacterMoverComp = Cast<UChaosCharacterMoverComponent>(&MoverComp))
+	{
+		ChaosCharacterMoverComp->ProduceLocalInput(Simulation->GetLocalSimInput_Mutable());
+	}
+
 	if (MoverComp.OnPreSimulationTick.IsBound())
 	{
 		MoverComp.OnPreSimulationTick.Broadcast(TimeStep, InputData.InputCmd);
diff --git a/Engine/Plugins/Experimental/ChaosMover/Source/ChaosMover/Private/ChaosMoverSimulation.cpp b/Engine/Plugins/Experimental/ChaosMover/Source/ChaosMover/Private/ChaosMoverSimulation.cpp
--- a/Engine/Plugins/Experimental/ChaosMover/Source/ChaosMover/Private/ChaosMoverSimulation.cpp
+++ b/Engine/Plugins/Experimental/ChaosMover/Source/ChaosMover/Private/ChaosMoverSimulation.cpp
@@ -189,6 +189,190 @@ void UChaosMoverSimulation::Deinit()
 	OnDeinit();
 }
 
+void UChaosMoverSimulation::BackupState(FChaosMoverPredictionStateBackup& Out) const
+{
+	// Todo @Harsha Extend this to a generic state back up flow
+	const UMoverBlackboard* BB = GetBlackboard();
+	if (!BB)
+	{
+		return;
+	}
+
+	// Reset to prevent stale values from a previous backup polluting this one
+	Out = FChaosMoverPredictionStateBackup();
+
+	FFloorCheckResult FloorResult;
+	if (BB->TryGet(CommonBlackboard::LastFloorResult, FloorResult))
+	{
+		Out.BlackboardLastFloorResult = FloorResult;
+	}
+
+	FWaterCheckResult WaterResult;
+	if (BB->TryGet(CommonBlackboard::LastWaterResult, WaterResult))
+	{
+		Out.BlackboardLastWaterResult = WaterResult;
+	}
+
+	UE::ChaosMover::FGroundDynamicsInfo GroundInfo;
+	if (BB->TryGet(UE::ChaosMover::Blackboard::GroundDynamicsInfo, GroundInfo))
+	{
+		Out.BlackboardGroundDynamicsInfo = GroundInfo;
+	}
+}
+
+void UChaosMoverSimulation::RestoreState(const FChaosMoverPredictionStateBackup& In)
+{
+	UMoverBlackboard* BB = GetBlackboard_Mutable();
+	if (!BB)
+	{
+		return;
+	}
+
+	if (In.BlackboardLastFloorResult.IsSet())
+	{
+		BB->Set(CommonBlackboard::LastFloorResult, In.BlackboardLastFloorResult.GetValue());
+	}
+	else
+	{
+		BB->Invalidate(CommonBlackboard::LastFloorResult);
+	}
+
+	if (In.BlackboardLastWaterResult.IsSet())
+	{
+		BB->Set(CommonBlackboard::LastWaterResult, In.BlackboardLastWaterResult.GetValue());
+	}
+	else
+	{
+		BB->Invalidate(CommonBlackboard::LastWaterResult);
+	}
+
+	if (In.BlackboardGroundDynamicsInfo.IsSet())
+	{
+		BB->Set(UE::ChaosMover::Blackboard::GroundDynamicsInfo, In.BlackboardGroundDynamicsInfo.GetValue());
+	}
+	else
+	{
+		BB->Invalidate(UE::ChaosMover::Blackboard::GroundDynamicsInfo);
+	}
+}
+
+void UChaosMoverSimulation::PredictTrajectory(const FMoverTimeStep& TimeStep, const UE::ChaosMover::FSimulationInputData& InputData, const FChaosMoverTrajectoryPredictionInputs* PredInputs, const FMoverSyncState& StartSyncState, UE::ChaosMover::FSimulationOutputData& OutputData) const
+{
+	Chaos::EnsureIsInPhysicsThreadContext();
+
+	constexpr int32 MinNumSteps = 1;
+	constexpr int32 MaxNumSteps = 120;
+	constexpr float MinDt = 0.001f;
+
+	// Using physics fixed dt
+	const float Dt = TimeStep.StepMs * 0.001f;
+
+	int32 NumSteps = 0;
+	if (PredInputs)
+	{
+		// Requested horizon is NumPredictionSteps * SecondsPerStep (seconds). Cover at least that amount of time with steps of size Dt.
+		const float RequestedHorizonSec = static_cast<float>(PredInputs->NumPredictionSteps) * PredInputs->SecondsPerStep;
+		NumSteps = FMath::Clamp(FMath::CeilToInt(RequestedHorizonSec / Dt), MinNumSteps, MaxNumSteps);
+		
+		const float NewHorizonSec = NumSteps * Dt;
+		if (NewHorizonSec < RequestedHorizonSec - UE_KINDA_SMALL_NUMBER)
+		{
+			UE_LOGF(LogChaosMover, Warning, "Predicted trajectory shorter than requested %.2f s, limited to %.2f s.",
+				RequestedHorizonSec, NewHorizonSec);
+		}
+	}
+	
+	if (FMath::IsNearlyZero(Dt) || NumSteps == 0)
+	{
+		UE_LOGF(LogChaosMover, Warning, "Invalid predicted trajectory parameters detected, exiting trajectory prediction.");
+		return;
+	}
+
+	const TStrongObjectPtr<const UBaseMovementMode> ModePtr = FindMovementModeByName(StartSyncState.MovementMode).Pin();
+	const UBaseMovementMode* Mode = ModePtr.Get();
+	if (!Mode)
+	{
+		UE_LOGF(LogChaosMover, Warning, "No active movement mode found, exiting trajectory prediction.");
+		return;
+	}
+
+	FMoverTickStartData StepState;
+	StepState.InputCmd = InputData.InputCmd;
+	StepState.SyncState = StartSyncState;
+	StepState.AuxState = InputData.AuxInputState;
+
+	FMoverDefaultSyncState* StepDefaultSyncState = StepState.SyncState.SyncStateCollection.FindMutableDataByType<FMoverDefaultSyncState>();
+	if (!StepDefaultSyncState)
+	{
+		return;
+	}
+
+	FMoverTimeStep FutureTimeStep = TimeStep;
+	FutureTimeStep.StepMs = Dt * 1000.0f;
+
+	FChaosMoverPredictedTrajectoryData& TrajOut = OutputData.AdditionalOutputData.FindOrAddMutableDataByType<FChaosMoverPredictedTrajectoryData>();
+	TrajOut.Deltas.SetNumUninitialized(NumSteps + 1);
+
+	FTransform PrevWorldTransform(StepDefaultSyncState->GetOrientation_WorldSpace(), StepDefaultSyncState->GetLocation_WorldSpace());
+	TrajOut.BaseTransform = PrevWorldTransform;
+	FVector PriorVelocity = StepDefaultSyncState->GetVelocity_WorldSpace();
+	FRotator PriorOrientation = StepDefaultSyncState->GetOrientation_WorldSpace();
+
+	for (int32 i = 0; i < NumSteps + 1; ++i)
+	{
+		FTrajectorySampleInfo& Sample = TrajOut.Deltas[i];
+		const FTransform CurrentWorldTransform(StepDefaultSyncState->GetOrientation_WorldSpace(), StepDefaultSyncState->GetLocation_WorldSpace());
+		if (i == 0)
+		{
+			Sample.Transform = FTransform::Identity;
+		}
+		else
+		{
+			Sample.Transform = CurrentWorldTransform.GetRelativeTransform(PrevWorldTransform);
+		}
+		Sample.LinearVelocity = StepDefaultSyncState->GetVelocity_WorldSpace();
+		Sample.SimTimeMs = FutureTimeStep.BaseSimTimeMs;
+
+		if (i == 0)
+		{
+			Sample.InstantaneousAcceleration = FVector::ZeroVector;
+			Sample.AngularVelocity = FRotator::ZeroRotator;
+		}
+		else
+		{
+			Sample.InstantaneousAcceleration = (Sample.LinearVelocity - PriorVelocity) / Dt;
+			Sample.AngularVelocity = (StepDefaultSyncState->GetOrientation_WorldSpace() - PriorOrientation) * (1.f / Dt);
+		}
+
+		PrevWorldTransform = CurrentWorldTransform;
+		PriorVelocity = Sample.LinearVelocity;
+		PriorOrientation = StepDefaultSyncState->GetOrientation_WorldSpace();
+
+		if (i == NumSteps)
+		{
+			break;
+		}
+
+		FProposedMove Proposed;
+		Mode->GenerateMove(StepState, FutureTimeStep, Proposed);
+
+		// Advance step (simple integration)
+		const FVector NextPos = StepDefaultSyncState->GetLocation_WorldSpace() + (Proposed.LinearVelocity * Dt);
+		const FRotator NextRot = UMovementUtils::ApplyAngularVelocityToRotator(StepDefaultSyncState->GetOrientation_WorldSpace(), Proposed.AngularVelocityDegrees, Dt);
+
+		StepDefaultSyncState->SetTransforms_WorldSpace(
+			NextPos,
+			NextRot,
+			Proposed.LinearVelocity,
+			Proposed.AngularVelocityDegrees,
+			StepDefaultSyncState->GetMovementBase(),
+			StepDefaultSyncState->GetMovementBaseBoneName());
+
+		FutureTimeStep.BaseSimTimeMs += FutureTimeStep.StepMs;
+		++FutureTimeStep.ServerFrame;
+	}
+}
+
 void UChaosMoverSimulation::OnInit()
 {
 	
@@ -470,6 +654,15 @@ void UChaosMoverSimulation::OnSimulationTick(const FMoverTimeStep& TimeStep, con
 
 	StateMachine.OnSimulationTick(TimeStep, TickStartData, Blackboard.Get(), MovementMixerWeakPtr.Get(), TickEndData);
 
+	if (const FChaosMoverTrajectoryPredictionInputs* PredInputs = LocalSimInput.FindDataByType<FChaosMoverTrajectoryPredictionInputs>())
+	{
+		// PredictTrajectory() can mutate state; so we backup/restore state around PredictTrajectory()
+		FChaosMoverPredictionStateBackup PredictionStateBackup;
+		BackupState(PredictionStateBackup);
+		PredictTrajectory(TimeStep, InputData, PredInputs, TickEndData.SyncState, OutputData);
+		RestoreState(PredictionStateBackup);
+	}
+
 	// Copy the sync state locally and to the output data
 	OutputData.SyncState = TickEndData.SyncState;
 	OutputData.LastUsedInputCmd = InputData.InputCmd;
diff --git a/Engine/Plugins/Experimental/ChaosMover/Source/ChaosMover/Private/Character/ChaosCharacterMoverComponent.cpp b/Engine/Plugins/Experimental/ChaosMover/Source/ChaosMover/Private/Character/ChaosCharacterMoverComponent.cpp
--- a/Engine/Plugins/Experimental/ChaosMover/Source/ChaosMover/Private/Character/ChaosCharacterMoverComponent.cpp
+++ b/Engine/Plugins/Experimental/ChaosMover/Source/ChaosMover/Private/Character/ChaosCharacterMoverComponent.cpp
@@ -29,6 +29,94 @@ UChaosCharacterMoverComponent::UChaosCharacterMoverComponent()
 	BackendClass = UChaosMoverBackendComponent::StaticClass();
 }
 
+TArray<FTrajectorySampleInfo> UChaosCharacterMoverComponent::GetPredictedTrajectory(FMoverPredictTrajectoryParams PredictionParams)
+{
+	const int32 NumSamples = FMath::Max(1, PredictionParams.NumPredictionSamples);
+	const float DtSeconds = FMath::Max(0.001f, PredictionParams.SecondsPerSample);
+	const TArray<FTrajectorySampleInfo>& CachedDeltas = LatestPredictedTrajectoryData.Deltas;
+
+	SetPredictedTrajectoryParams(NumSamples - 1, DtSeconds, /*bInEnableTrajectoryPrediction*/ true);
+
+	TArray<FTrajectorySampleInfo> Out;
+
+	if (!LastMoverDefaultSyncState)
+	{
+		Out.AddDefaulted(NumSamples);
+	}
+	else
+	{
+		FTrajectorySampleInfo BaseSample;
+		const FVector AngVelDeg = LastMoverDefaultSyncState->GetAngularVelocityDegrees_WorldSpace();
+		BaseSample.Transform = LastMoverDefaultSyncState->GetTransform_WorldSpace();
+		BaseSample.LinearVelocity = LastMoverDefaultSyncState->GetVelocity_WorldSpace();
+		BaseSample.AngularVelocity = FRotator(AngVelDeg.X, AngVelDeg.Y, AngVelDeg.Z);
+		BaseSample.InstantaneousAcceleration = FVector::ZeroVector;
+		BaseSample.SimTimeMs = GetLastTimeStep().BaseSimTimeMs;
+
+		if (CachedDeltas.Num() < 2)
+		{
+			Out.SetNumUninitialized(NumSamples);
+			for (int32 i = 0; i < NumSamples; ++i)
+			{
+				Out[i] = BaseSample;
+				Out[i].SimTimeMs = BaseSample.SimTimeMs + static_cast<double>(i) * static_cast<double>(DtSeconds) * 1000.0;
+			}
+		}
+		else
+		{
+			const int32 NumDeltaSamples = CachedDeltas.Num();
+			Out.SetNumUninitialized(NumDeltaSamples + 1);
+			Out[0] = BaseSample;
+
+			Out[1] = CachedDeltas[0];
+			Out[1].Transform = LatestPredictedTrajectoryData.BaseTransform;
+			for (int32 i = 1; i < NumDeltaSamples; ++i)
+			{
+				Out[i + 1] = CachedDeltas[i];
+				Out[i + 1].Transform = CachedDeltas[i].Transform * Out[i].Transform;
+			}
+		}
+	}
+
+	if (PredictionParams.bUseVisualComponentRoot)
+	{
+		if (const USceneComponent* VisualComp = GetPrimaryVisualComponent())
+		{
+			const FTransform VisualRelativeTransform = VisualComp->GetRelativeTransform();
+			for (FTrajectorySampleInfo& Sample : Out)
+			{
+				Sample.Transform = VisualRelativeTransform * Sample.Transform;
+			}
+		}
+	}
+
+	return Out;
+}
+
+void UChaosCharacterMoverComponent::SetPredictedTrajectoryParams(int32 InSteps, float InStepSeconds, bool bInEnableTrajectoryPrediction)
+{
+	PredictedTrajectorySteps = FMath::Max(1, InSteps);
+	PredictedTrajectoryStepSeconds = FMath::Max(0.001f, InStepSeconds);
+	bEnableTrajectoryPrediction = bInEnableTrajectoryPrediction;
+}
+
+void UChaosCharacterMoverComponent::GetPredictedTrajectoryParams(int32& OutSteps, float& OutStepSeconds, bool& bIsTrajectoryPredictionEnabled) const
+{
+	OutSteps = PredictedTrajectorySteps;
+	OutStepSeconds = PredictedTrajectoryStepSeconds;
+	bIsTrajectoryPredictionEnabled = bEnableTrajectoryPrediction;
+}
+
+void UChaosCharacterMoverComponent::ProduceLocalInput(FMoverDataCollection& OutLocalSimInput) const
+{
+	if (bEnableTrajectoryPrediction)
+	{
+		FChaosMoverTrajectoryPredictionInputs& PredInputs = OutLocalSimInput.FindOrAddMutableDataByType<FChaosMoverTrajectoryPredictionInputs>();
+		bool bUnused;
+		GetPredictedTrajectoryParams(PredInputs.NumPredictionSteps, PredInputs.SecondsPerStep, bUnused);
+	}
+}
+
 void UChaosCharacterMoverComponent::OnMoverPreSimulationTick(const FMoverTimeStep& TimeStep, const FMoverInputCmdContext& InputCmd)
 {
 	// This overrides the base class to make sure that it doesn't do anything. We do this
@@ -99,6 +187,12 @@ void UChaosCharacterMoverComponent::SetAdditionalSimulationOutput(const FMoverDa
 		bWaterResultSet = true;
 		LatestWaterResult = WaterData->WaterResult;
 	}
+
+	if (const FChaosMoverPredictedTrajectoryData* PredictedTraj = Data.FindDataByType<FChaosMoverPredictedTrajectoryData>())
+	{
+		bPredictedTrajectorySet = true;
+		LatestPredictedTrajectoryData = *PredictedTraj;
+	}
 }
 
 bool UChaosCharacterMoverComponent::TryGetLastWaterResult(FWaterCheckResult& OutWaterResult) const
diff --git a/Engine/Plugins/Experimental/ChaosMover/Source/ChaosMover/Public/ChaosMover/ChaosMoverSimulation.h b/Engine/Plugins/Experimental/ChaosMover/Source/ChaosMover/Public/ChaosMover/ChaosMoverSimulation.h
--- a/Engine/Plugins/Experimental/ChaosMover/Source/ChaosMover/Public/ChaosMover/ChaosMoverSimulation.h
+++ b/Engine/Plugins/Experimental/ChaosMover/Source/ChaosMover/Public/ChaosMover/ChaosMoverSimulation.h
@@ -6,7 +6,9 @@
 #include "ChaosMover/ChaosMoverSimulationTypes.h"
 #include "ChaosMover/PathedMovement/ChaosPathedMovementTypes.h"
 #include "ChaosMover/ChaosMoverStateMachine.h"
+#include "MoveLibrary/FloorQueryUtils.h"
 #include "MoveLibrary/MovementMixer.h"
+#include "MoveLibrary/WaterMovementUtils.h"
 #include "MoverTypes.h"
 #include "MoverSimulation.h"
 #include "PhysicsEngine/ConstraintInstance.h"
@@ -21,6 +23,13 @@ namespace Chaos
 	class FCollisionContactModifier;
 }
 
+struct FChaosMoverPredictionStateBackup
+{
+	TOptional<FFloorCheckResult> BlackboardLastFloorResult;
+	TOptional<FWaterCheckResult> BlackboardLastWaterResult;
+	TOptional<UE::ChaosMover::FGroundDynamicsInfo> BlackboardGroundDynamicsInfo;
+};
+
 UCLASS(MinimalAPI, BlueprintType)
 class UChaosMoverSimulation : public UMoverSimulation
 {
@@ -153,6 +162,11 @@ protected:
 	CHAOSMOVER_API virtual void ProcessSimulationEvent(const FMoverSimulationEventData& ModeChangedData);
 	CHAOSMOVER_API virtual void OnMovementModeChanged(const FMovementModeChangedEventData& ModeChangedData);
 
+	// Trajectory Prediction
+	CHAOSMOVER_API void BackupState(FChaosMoverPredictionStateBackup& Out) const;
+	CHAOSMOVER_API void RestoreState(const FChaosMoverPredictionStateBackup& In);
+	CHAOSMOVER_API void PredictTrajectory(const FMoverTimeStep& TimeStep, const UE::ChaosMover::FSimulationInputData& InputData, const FChaosMoverTrajectoryPredictionInputs* PredInputs, const FMoverSyncState& StartSyncState, UE::ChaosMover::FSimulationOutputData& OutputData) const;
+
 	// Teleportation (internal functions)
 public:
 	// Attempt to teleport. This will first check CanTeleport and broadcast simulation events to indicate success or failure.
diff --git a/Engine/Plugins/Experimental/ChaosMover/Source/ChaosMover/Public/ChaosMover/ChaosMoverSimulationTypes.h b/Engine/Plugins/Experimental/ChaosMover/Source/ChaosMover/Public/ChaosMover/ChaosMoverSimulationTypes.h
--- a/Engine/Plugins/Experimental/ChaosMover/Source/ChaosMover/Public/ChaosMover/ChaosMoverSimulationTypes.h
+++ b/Engine/Plugins/Experimental/ChaosMover/Source/ChaosMover/Public/ChaosMover/ChaosMoverSimulationTypes.h
@@ -8,6 +8,7 @@
 #include "CollisionQueryParams.h"
 #include "DefaultMovementSet/MovementModifiers/StanceModifier.h"
 #include "MoveLibrary/FloorQueryUtils.h"
+#include "MoveLibrary/MovementUtils.h"
 #include "MoveLibrary/WaterMovementUtils.h"
 #include "MoverSimulationTypes.h"
 #include "LayeredMove.h"
@@ -73,6 +74,98 @@ namespace UE::ChaosMover
 
 } // namespace UE::ChaosMover
 
+// Predicted trajectory deltas generated during simulation tick.
+// Stored in FSimulationOutputData::AdditionalOutputData (not networked, local only).
+USTRUCT(BlueprintType)
+struct FChaosMoverPredictedTrajectoryData : public FMoverDataStructBase
+{
+	GENERATED_BODY()
+
+	UPROPERTY(BlueprintReadWrite, Category = ChaosMover)
+	FTransform BaseTransform = FTransform::Identity;
+
+	// Predicted trajectory as delta transforms relative to each preceding step.
+	// Element [0] is the starting transform (time offset 0, Identity), i.e., anchored at BaseTransform.
+	UPROPERTY(BlueprintReadWrite, Category = ChaosMover)
+	TArray<FTrajectorySampleInfo> Deltas;
+
+	virtual UScriptStruct* GetScriptStruct() const override
+	{
+		return StaticStruct();
+	}
+
+	virtual FMoverDataStructBase* Clone() const override
+	{
+		return new FChaosMoverPredictedTrajectoryData(*this);
+	}
+
+	virtual bool ShouldReconcile(const FMoverDataStructBase& AuthorityState) const override
+	{
+		return false;
+	}
+
+	virtual void Interpolate(const FMoverDataStructBase& From, const FMoverDataStructBase& To, float Pct) override
+	{
+		// Always take the latest trajectory rather than blending so that downstream code can be responsive to predictions
+		*this = static_cast<const FChaosMoverPredictedTrajectoryData&>(To);
+	}
+
+	virtual void Merge(const FMoverDataStructBase& From) override {}
+};
+
+template<>
+struct TStructOpsTypeTraits<FChaosMoverPredictedTrajectoryData> : public TStructOpsTypeTraitsBase2<FChaosMoverPredictedTrajectoryData>
+{
+	enum
+	{
+		WithCopy = true
+	};
+};
+
+// Local-only inputs that control how ChaosMover generates predicted trajectory samples.
+USTRUCT(BlueprintType)
+struct FChaosMoverTrajectoryPredictionInputs : public FMoverDataStructBase
+{
+	GENERATED_BODY()
+
+	UPROPERTY(EditAnywhere, BlueprintReadWrite, Category = ChaosMover, meta = (ClampMin = "1", UIMin = "1"))
+	int32 NumPredictionSteps = 15;
+
+	UPROPERTY(EditAnywhere, BlueprintReadWrite, Category = ChaosMover, meta = (ClampMin = "0.001", UIMin = "0.001"))
+	float SecondsPerStep = 0.1f;
+
+	virtual UScriptStruct* GetScriptStruct() const override
+	{
+		return StaticStruct();
+	}
+
+	virtual FMoverDataStructBase* Clone() const override
+	{
+		return new FChaosMoverTrajectoryPredictionInputs(*this);
+	}
+
+	virtual bool ShouldReconcile(const FMoverDataStructBase& AuthorityState) const override
+	{
+		return false;
+	}
+
+	virtual void Interpolate(const FMoverDataStructBase& From, const FMoverDataStructBase& To, float Pct) override
+	{
+		*this = static_cast<const FChaosMoverTrajectoryPredictionInputs&>((Pct < 0.5f) ? From : To);
+	}
+
+	virtual void Merge(const FMoverDataStructBase& From) override {}
+};
+
+template<>
+struct TStructOpsTypeTraits<FChaosMoverTrajectoryPredictionInputs> : public TStructOpsTypeTraitsBase2<FChaosMoverTrajectoryPredictionInputs>
+{
+	enum
+	{
+		WithCopy = true
+	};
+};
+
 // Default chaos mover character simulation state, contains state basic for character simulation
 USTRUCT(BlueprintType)
 struct FChaosMoverCharacterSimState : public FMoverDataStructBase
diff --git a/Engine/Plugins/Experimental/ChaosMover/Source/ChaosMover/Public/ChaosMover/Character/ChaosCharacterMoverComponent.h b/Engine/Plugins/Experimental/ChaosMover/Source/ChaosMover/Public/ChaosMover/Character/ChaosCharacterMoverComponent.h
--- a/Engine/Plugins/Experimental/ChaosMover/Source/ChaosMover/Public/ChaosMover/Character/ChaosCharacterMoverComponent.h
+++ b/Engine/Plugins/Experimental/ChaosMover/Source/ChaosMover/Public/ChaosMover/Character/ChaosCharacterMoverComponent.h
@@ -2,6 +2,7 @@
 
 #pragma once
 
+#include "Containers/ArrayView.h"
 #include "ChaosMover/ChaosMoverSimulationTypes.h"
 #include "DefaultMovementSet/CharacterMoverComponent.h"
 #include "DefaultMovementSet/CharacterMoverSimulationTypes.h"
@@ -31,6 +32,7 @@ public:
 
 	CHAOSMOVER_API virtual void ProduceInput(const int32 DeltaTimeMS, FMoverInputCmdContext* Cmd) override;
 	CHAOSMOVER_API virtual void ProduceServerInput(const int32 DeltaTimeMS, FMoverInputCmdContext* Cmd);
+	CHAOSMOVER_API virtual void ProduceLocalInput(FMoverDataCollection& OutLocalSimInput) const;
 	CHAOSMOVER_API virtual void DoQueueNextMode(FName DesiredModeName, bool bShouldReenter=false) override;
 	CHAOSMOVER_API virtual FName GetNextMovementModeName() const override;
 
@@ -60,6 +62,18 @@ public:
 	UFUNCTION(BlueprintCallable, Category = "Chaos Mover")
 	void CancelMovementSettingsOverrides(FName ModeName = NAME_None);
 
+	// Get cached predicted trajectory deltas.
+	// The first element is Identity corresponding to the starting transform at time offset 0.
+	TConstArrayView<FTrajectorySampleInfo> GetCachedPredictedTrajectoryDeltas() const { return LatestPredictedTrajectoryData.Deltas; }
+
+	// Returns predicted trajectory constructed from the cached deltas.
+	// The first element is the starting transform at time offset 0.
+	// Supported parameters: bUseVisualComponentRoot, NumPredictionSamples and SecondsPerSample.
+	CHAOSMOVER_API virtual TArray<FTrajectorySampleInfo> GetPredictedTrajectory(FMoverPredictTrajectoryParams PredictionParams) override;
+
+	CHAOSMOVER_API void SetPredictedTrajectoryParams(int32 InSteps, float InStepSeconds, bool bInEnableTrajectoryPrediction);
+	CHAOSMOVER_API void GetPredictedTrajectoryParams(int32& OutSteps, float& OutStepSeconds, bool& bIsTrajectoryPredictionEnabled) const;
+
 protected:
 	CHAOSMOVER_API virtual void OnMoverPreSimulationTick(const FMoverTimeStep& TimeStep, const FMoverInputCmdContext& InputCmd) override;
 	CHAOSMOVER_API virtual void ProcessSimulationEvent(const FMoverSimulationEventData& EventData) override;
@@ -85,6 +99,13 @@ protected:
 	bool bWaterResultSet = false;
 	FWaterCheckResult LatestWaterResult;
 
+	bool bPredictedTrajectorySet = false;
+	FChaosMoverPredictedTrajectoryData LatestPredictedTrajectoryData;
+
+	int32 PredictedTrajectorySteps = 15;
+	float PredictedTrajectoryStepSeconds = 0.1f;
+	bool bEnableTrajectoryPrediction = false;
+
 	FVector LaunchVelocityOrImpulse = FVector::ZeroVector;
 	EChaosMoverVelocityEffectMode LaunchMode = EChaosMoverVelocityEffectMode::AdditiveVelocity;
 
diff --git a/Engine/Plugins/Experimental/Mover/Source/Mover/Private/MoverPoseSearchTrajectoryPredictor.cpp b/Engine/Plugins/Experimental/Mover/Source/Mover/Private/MoverPoseSearchTrajectoryPredictor.cpp
--- a/Engine/Plugins/Experimental/Mover/Source/Mover/Private/MoverPoseSearchTrajectoryPredictor.cpp
+++ b/Engine/Plugins/Experimental/Mover/Source/Mover/Private/MoverPoseSearchTrajectoryPredictor.cpp
@@ -49,14 +49,21 @@ void UMoverTrajectoryPredictor::Predict(UMoverComponent& MoverComponent, FTransf
 		float CurrentTime = InOutTrajectory.Samples[NumHistorySamples].TimeInSeconds;	// t == 0 is the last frame, so we need to account for the starting mover state being offset
 		FTransform MoverPredictLower = MoverPredictionSamples[0].Transform;
 		FTransform MoverPredictUpper = MoverPredictionSamples[1].Transform;
-		float TimeInSecondsLower = CurrentTime;
-		float TimeInSecondsUpper = CurrentTime + MoverSamplingInterval;
+
+		// Use actual SimTimeMs from samples for bracket calculation to account for non-uniformly spaced MoverPredictionSamples.
+		auto SampleTimeToSeconds = [&](int32 Idx) -> float
+		{
+			return (MoverPredictionSamples[Idx].SimTimeMs - MoverPredictionSamples[0].SimTimeMs)*0.001f + CurrentTime;
+		};
+
+		float TimeInSecondsLower = SampleTimeToSeconds(0); // == CurrentTime for sample 0
+		float TimeInSecondsUpper = SampleTimeToSeconds(1);
 		float AccumulatedSeconds = CurrentTime + SecondsPerPredictionSample; // first prediction sample should be at time for the first prediction
 
 		for (int32 i = 0; i < NumPredictionSamples; ++i)
 		{
 			// Progress target if necessary
-			while (AccumulatedSeconds > TimeInSecondsUpper && MoverPredictionSampleIndex < NumMoverPredictionSamplesRequired - 1)
+			while (AccumulatedSeconds > TimeInSecondsUpper && MoverPredictionSampleIndex < MoverPredictionSamples.Num() - 1)
 			{
 				MoverPredictionSampleIndex++;
 
@@ -64,8 +71,8 @@ void UMoverTrajectoryPredictor::Predict(UMoverComponent& MoverComponent, FTransf
 				MoverPredictLower = MoverPredictionSamples[MoverPredictionSampleIndex - 1].Transform;
 				MoverPredictUpper = MoverPredictionSamples[MoverPredictionSampleIndex].Transform;
 
-				TimeInSecondsLower = ((MoverPredictionSampleIndex - 1) * MoverSamplingInterval) + CurrentTime;
-				TimeInSecondsUpper = (MoverPredictionSampleIndex * MoverSamplingInterval) + CurrentTime;
+				TimeInSecondsLower = SampleTimeToSeconds(MoverPredictionSampleIndex - 1);
+				TimeInSecondsUpper = SampleTimeToSeconds(MoverPredictionSampleIndex);
 			}
 
 			const int PoseSampleIdx = i + NumHistorySamples + 1;
diff --git a/Engine/Plugins/Experimental/Mover/Source/Mover/Public/MoverComponent.h b/Engine/Plugins/Experimental/Mover/Source/Mover/Public/MoverComponent.h
--- a/Engine/Plugins/Experimental/Mover/Source/Mover/Public/MoverComponent.h
+++ b/Engine/Plugins/Experimental/Mover/Source/Mover/Public/MoverComponent.h
@@ -518,7 +518,7 @@ public:	// Queries
 	/** Get a sampling of where the actor is projected to be in the future, based on a current state. Note that this is projecting ideal movement without doing full simulation and collision.
 	 * The first sample info of the returned array corresponds to the current state of the mover. */
 	UFUNCTION(BlueprintCallable, BlueprintPure = false, Category = Mover)
-	MOVER_API TArray<FTrajectorySampleInfo> GetPredictedTrajectory(FMoverPredictTrajectoryParams PredictionParams);	
+	MOVER_API virtual TArray<FTrajectorySampleInfo> GetPredictedTrajectory(FMoverPredictTrajectoryParams PredictionParams);
 
 	// Get the current movement mode name
 	UFUNCTION(BlueprintPure, Category = Mover)
