ChaosMover: Fixing that easing curves were not being applied to path or path patterns correctly

#Synced-CL 51926426

diff --git a/Engine/Plugins/Experimental/ChaosMover/Source/ChaosMover/Private/PathedMovement/ChaosPathedMovementMode.cpp b/Engine/Plugins/Experimental/ChaosMover/Source/ChaosMover/Private/PathedMovement/ChaosPathedMovementMode.cpp
--- a/Engine/Plugins/Experimental/ChaosMover/Source/ChaosMover/Private/PathedMovement/ChaosPathedMovementMode.cpp
+++ b/Engine/Plugins/Experimental/ChaosMover/Source/ChaosMover/Private/PathedMovement/ChaosPathedMovementMode.cpp
@@ -463,6 +463,8 @@ FTransform UChaosPathedMovementMode::CalcTargetTransform(float OverallPathProgre
 	FTransform BasisTransformNoScale = BasisTransform;
 	BasisTransformNoScale.SetScale3D(FVector(1,1,1));
 
+	const float EasedPathProgress = FAlphaBlend::AlphaToBlendOption(OverallPathProgress, Easing, CustomEasingCurve);
+
 	FTransform TargetRelativeTransform = FTransform::Identity;
 	FVector AccumulatedRelativeScale = FVector(1,1,1);
 
@@ -470,7 +472,7 @@ FTransform UChaosPathedMovementMode::CalcTargetTransform(float OverallPathProgre
 	{
 		if (PathPattern)
 		{
-			const FTransform PatternWorldTransform = PathPattern->CalcTargetTransform(OverallPathProgress, BasisTransformNoScale);
+			const FTransform PatternWorldTransform = PathPattern->CalcTargetTransform(EasedPathProgress, BasisTransformNoScale);
 			const FTransform PatternRelativeTransform = FTransform(
 				BasisTransformNoScale.InverseTransformRotation(PatternWorldTransform.GetRotation()),
 				BasisTransformNoScale.InverseTransformPositionNoScale(PatternWorldTransform.GetLocation()));
diff --git a/Engine/Plugins/Experimental/ChaosMover/Source/ChaosMover/Private/PathedMovement/ChaosPathedMovementPatternBase.cpp b/Engine/Plugins/Experimental/ChaosMover/Source/ChaosMover/Private/PathedMovement/ChaosPathedMovementPatternBase.cpp
--- a/Engine/Plugins/Experimental/ChaosMover/Source/ChaosMover/Private/PathedMovement/ChaosPathedMovementPatternBase.cpp
+++ b/Engine/Plugins/Experimental/ChaosMover/Source/ChaosMover/Private/PathedMovement/ChaosPathedMovementPatternBase.cpp
@@ -2,6 +2,7 @@
 
 #include "ChaosMover/PathedMovement/ChaosPathedMovementPatternBase.h"
 
+#include "AlphaBlend.h"
 #include "HAL/IConsoleManager.h"
 #include "ChaosMover/PathedMovement/ChaosPathedMovementMode.h"
 #include "ChaosMover/ChaosMoverSimulation.h"
@@ -127,6 +128,10 @@ float UChaosPathedMovementPatternBase::ConvertPatternToPathProgress(float FinalP
 FTransform UChaosPathedMovementPatternBase::CalcTargetTransform(float OverallPathProgress, const FTransform& BasisTransform) const
 {
 	float PatternProgress = ConvertPathToPatternProgress(OverallPathProgress);
+	if (PatternProgress >= 0.0f)
+	{
+		PatternProgress = FAlphaBlend::AlphaToBlendOption(PatternProgress, Easing, CustomEasingCurve);
+	}
 	return CalcMaskedTargetTransform(PatternProgress, BasisTransform);
 }
 
