ChaosMover: Fix Pathed Movement loop offset issues. Add a BounceAtEnds bool, useful for non looping paths when we want to prevent teleports when the path progress wraps around (which happens when using path offset). Wrap around is desirable for looping paths, but creates teleports for non looping paths (might not always be what the user wants)

#Synced-CL 51926426

diff --git a/Engine/Plugins/Experimental/ChaosMover/Source/ChaosMover/Private/PathedMovement/ChaosPathedMovementPatternBase.cpp b/Engine/Plugins/Experimental/ChaosMover/Source/ChaosMover/Private/PathedMovement/ChaosPathedMovementPatternBase.cpp
--- a/Engine/Plugins/Experimental/ChaosMover/Source/ChaosMover/Private/PathedMovement/ChaosPathedMovementPatternBase.cpp
+++ b/Engine/Plugins/Experimental/ChaosMover/Source/ChaosMover/Private/PathedMovement/ChaosPathedMovementPatternBase.cpp
@@ -75,12 +75,54 @@ float UChaosPathedMovementPatternBase::ConvertPathToPatternProgress(float Overal
 		CurLoopPathProgress = PathProgressPerPatternLoop;
 	}
 
-	// If each loop is there and back, the actual path progress needs to flip after completing half the span
+	// Normalized progress within this loop, in [0, 1]
+	const float t = CurLoopPathProgress / PathProgressPerPatternLoop;
+
+	if (bBounceAtEnds)
+	{
+		// Bounce at 0 and 1 instead of wrapping - safe for non-looping paths
+		if (!IsOneWay)
+		{
+			// Piecewise triangle spanning the full [0, 1] range. No seam crossing.
+			// Ping (first half):  rises  LoopOffset -> 1.0
+			// Middle transition:  falls  1.0 -> 0.0  (passes through LoopOffset on the way)
+			// Pong (second half): rises  0.0 -> LoopOffset
+			// LoopOffset=0: standard 0->1->0. LoopOffset=0.5: 0.5->1.0->0.0->0.5.
+			const float Leg1End = (1.f - LoopOffset) * 0.5f;
+			const float Leg2End = (2.f - LoopOffset) * 0.5f;
+			if (t <= Leg1End)
+			{
+				return LoopOffset + 2.f * t;
+			}
+			else if (t <= Leg2End)
+			{
+				return 2.f - LoopOffset - 2.f * t;
+			}
+			else
+			{
+				return 2.f * t - 2.f + LoopOffset;
+			}
+		}
+		else  // IsOneWay
+		{
+			// Triangle wave bouncing between LoopOffset and 1.0. Never goes below LoopOffset, never crosses seam.
+			// LoopOffset=0: 0->1->0. LoopOffset=0.5: 0.5->1.0->0.5.
+			if (t <= 0.5f)
+			{
+				return LoopOffset + 2.f * t * (1.f - LoopOffset);
+			}
+			else
+			{
+				return LoopOffset + 2.f * (1.f - t) * (1.f - LoopOffset);
+			}
+		}
+	}
+
+	// Seam-crossing (wraps at 0/1): if each loop is there and back, flip progress after completing half the span
 	if (!IsOneWay)
 	{
 		// ThereAndBack means we actually progress twice as fast as a OneShot
 		CurLoopPathProgress *= 2.f;
-		
 		const float ReversePathProgress = CurLoopPathProgress - PathProgressPerPatternLoop;
 		if (ReversePathProgress > 0.f)
 		{
@@ -108,21 +150,43 @@ float UChaosPathedMovementPatternBase::ConvertPatternToPathProgress(float FinalP
 		return -1.0f;
 	}
 
+	const float PathProgressPerPatternLoop = (EndAtPathProgress - StartAtPathProgress) / NumLoopsPerPath;
+
+	if (bBounceAtEnds)
+	{
+		// Compute normalized loop progress t in [0, 1] for the first occurrence of FinalPatternProgress.
+		float t;
+		if (!IsOneWay)
+		{
+			// Inverse of leg 1: P = LoopOffset + 2t  =>  t = (P - LoopOffset) / 2
+			t = FMath::Clamp((FMath::Clamp(FinalPatternProgress, LoopOffset, 1.0f) - LoopOffset) * 0.5f, 0.f, (1.f - LoopOffset) * 0.5f);
+		}
+		else  // IsOneWay - inverse of P = LoopOffset + 2t*(1-LoopOffset) for t in [0, 0.5]
+		{
+			const float Denom = 2.f * (1.f - LoopOffset);
+			t = FMath::IsNearlyZero(Denom) ? 0.f
+				: FMath::Clamp((FMath::Clamp(FinalPatternProgress, LoopOffset, 1.0f) - LoopOffset) / Denom, 0.f, 0.5f);
+		}
+		return FMath::Clamp(StartAtPathProgress + t * PathProgressPerPatternLoop, 0.0f, 1.0f);
+	}
+
+	// Seam-crossing: reverse the LoopOffset wrap, then reverse the IsOneWay flip
 	const float PatternProgressOffsetRemoved = FMath::Clamp(FinalPatternProgress, 0.0f, 1.0f) - LoopOffset;
 	float PatternProgress = FMath::Fmod(PatternProgressOffsetRemoved, 1.0f);
+	if (PatternProgress < 0.f) { PatternProgress += 1.f; }
 	if (FMath::IsNearlyZero(PatternProgress) && PatternProgressOffsetRemoved > UE_SMALL_NUMBER)
 	{
 		// Treat PatternProgress = 1.0f as 1.0f
 		PatternProgress = 1.0f;
 	}
-	const float PathProgressPerPatternLoop = (EndAtPathProgress - StartAtPathProgress) / NumLoopsPerPath;
+
 	float CurLoopPathProgress = PatternProgress * PathProgressPerPatternLoop;
 	if (!IsOneWay)
 	{
 		CurLoopPathProgress /= 2.f;
 	}
 
-	return FMath::Clamp(StartAtPathProgress + (EndAtPathProgress - StartAtPathProgress) * CurLoopPathProgress, 0.0f, 1.0f);
+	return FMath::Clamp(StartAtPathProgress + CurLoopPathProgress, 0.0f, 1.0f);
 }
 
 FTransform UChaosPathedMovementPatternBase::CalcTargetTransform(float OverallPathProgress, const FTransform& BasisTransform) const
diff --git a/Engine/Plugins/Experimental/ChaosMover/Source/ChaosMover/Public/ChaosMover/PathedMovement/ChaosPathedMovementPatternBase.h b/Engine/Plugins/Experimental/ChaosMover/Source/ChaosMover/Public/ChaosMover/PathedMovement/ChaosPathedMovementPatternBase.h
--- a/Engine/Plugins/Experimental/ChaosMover/Source/ChaosMover/Public/ChaosMover/PathedMovement/ChaosPathedMovementPatternBase.h
+++ b/Engine/Plugins/Experimental/ChaosMover/Source/ChaosMover/Public/ChaosMover/PathedMovement/ChaosPathedMovementPatternBase.h
@@ -127,9 +127,17 @@ protected:
 	UPROPERTY(EditAnywhere, Category = "Path Pattern", meta = (UIMin = 0, ClampMin = 0))
 	int32 NumLoopsPerPath = 1;
 
-	/** Whether a loop is played one way (0 -> 1) or there and back (0 -> 1 -> 0)*/
+	/** If true, each loop plays forward only (0 → 1). If false, each loop plays there-and-back (0 → 1 → 0). */
 	UPROPERTY(EditAnywhere, Category = "Path Pattern")
 	bool IsOneWay = true;
+
+	/**
+	 * When true, progress bounces at 0 and 1 instead of crossing the seam (wrapping).
+	 * Use this for non-looping paths to avoid teleports at the endpoints.
+	 * When false (default), progress wraps at the seam — correct for geometrically looping paths.
+	 */
+	UPROPERTY(EditAnywhere, Category = "Path Pattern")
+	bool bBounceAtEnds = false;
 	
 	/** 
 	 * If true, the component will be rotated to face in the direction of this pattern's motion.
