Making it easier to compose moves with modes by creating a move executor and a move based mode that can use a move executor on one hand and a move source on the other, which delegate the proposed move generation to a layered move

#Synced-CL 52152484

diff --git a/Engine/Plugins/Experimental/ChaosMover/Source/ChaosMover/Private/ChaosCompositeMovementMode.cpp b/Engine/Plugins/Experimental/ChaosMover/Source/ChaosMover/Private/ChaosCompositeMovementMode.cpp
new file mode 100644
--- /dev/null
+++ b/Engine/Plugins/Experimental/ChaosMover/Source/ChaosMover/Private/ChaosCompositeMovementMode.cpp
@@ -0,0 +1,101 @@
+// Copyright Epic Games, Inc. All Rights Reserved.
+
+#include "ChaosMover/ChaosCompositeMovementMode.h"
+#include "ChaosMover/ChaosMoverLog.h"
+#include "ChaosMover/ChaosMoverSimulation.h"
+#include "MoverSimulationTypes.h"
+
+#include UE_INLINE_GENERATED_CPP_BY_NAME(ChaosCompositeMovementMode)
+
+
+UChaosCompositeMovementMode::UChaosCompositeMovementMode(const FObjectInitializer& ObjectInitializer)
+	: Super(ObjectInitializer)
+{
+	bSupportsAsync = true;
+}
+
+void UChaosCompositeMovementMode::OnRegistered(const FName ModeName, const FMoverSimContext& SimContext)
+{
+	Super::OnRegistered(ModeName, SimContext);
+
+	if (MoveExecutor && Simulation)
+	{
+		MoveExecutor->SetSimulation(Simulation);
+	}
+
+	if (MoveExecutor)
+	{
+		MoveExecutor->OnModeRegistered(ModeName);
+	}
+}
+
+void UChaosCompositeMovementMode::OnUnregistered(const FMoverSimContext& SimContext)
+{
+	if (MoveExecutor)
+	{
+		MoveExecutor->OnModeUnregistered();
+		MoveExecutor->SetSimulation(nullptr);
+	}
+
+	Super::OnUnregistered(SimContext);
+}
+
+void UChaosCompositeMovementMode::Activate(const FMoverEventContext& Context, FName PrevModeName, const FMoverSimContext& SimContext, const FMoverTickStartData& StartState, FMoverSyncState* OutSyncState, FMoverAuxStateContext* OutAuxState)
+{
+	Super::Activate(Context, PrevModeName, SimContext, StartState, OutSyncState, OutAuxState);
+
+	bMoveSourceEnded = false;
+
+	if (MoveSource && Simulation)
+	{
+		MoveSource->OnStart_Async(Simulation->GetBlackboard_Mutable(), Context.EventTimeMs);
+	}
+}
+
+void UChaosCompositeMovementMode::Deactivate(const FMoverEventContext& Context, FName InNextModeName, const FMoverSimContext& SimContext)
+{
+	if (MoveSource && Simulation && !bMoveSourceEnded)
+	{
+		MoveSource->OnEnd_Async(Simulation->GetBlackboard_Mutable(), Context.EventTimeMs);
+	}
+
+	Super::Deactivate(Context, InNextModeName, SimContext);
+}
+
+void UChaosCompositeMovementMode::GenerateMove_Implementation(const FMoverSimContext& SimContext, const FMoverTickStartData& StartState, const FMoverTimeStep& TimeStep, FProposedMove& OutProposedMove) const
+{
+	if (!MoveSource)
+	{
+		return;
+	}
+
+	UMoverBlackboard* SimBlackboard = Simulation ? Simulation->GetBlackboard_Mutable() : nullptr;
+	MoveSource->GenerateMove_Async(StartState, TimeStep, SimBlackboard, OutProposedMove);
+}
+
+void UChaosCompositeMovementMode::SimulationTick_Implementation(const FSimulationTickParams& Params, FMoverTickEndData& OutputState)
+{
+	if (MoveExecutor)
+	{
+		MoveExecutor->ExecuteMove_Async(Params, OutputState);
+	}
+	else
+	{
+		UE_LOG(LogChaosMover, Warning, TEXT("UChaosCompositeMovementMode (%s): No MoveExecutor set. Mode will do nothing."), *GetName());
+	}
+
+	if (MoveSource)
+	{
+		const double CurrentSimTimeMs = Params.TimeStep.BaseSimTimeMs;
+		if (!bMoveSourceEnded && MoveSource->IsFinished(CurrentSimTimeMs))
+		{
+			bMoveSourceEnded = true;
+			MoveSource->OnEnd_Async(Params.SimBlackboard, CurrentSimTimeMs);
+
+			if (!NextModeName.IsNone())
+			{
+				OutputState.MovementEndState.NextModeName = NextModeName;
+			}
+		}
+	}
+}
diff --git a/Engine/Plugins/Experimental/ChaosMover/Source/ChaosMover/Private/ChaosLayeredMoveSource.cpp b/Engine/Plugins/Experimental/ChaosMover/Source/ChaosMover/Private/ChaosLayeredMoveSource.cpp
new file mode 100644
--- /dev/null
+++ b/Engine/Plugins/Experimental/ChaosMover/Source/ChaosMover/Private/ChaosLayeredMoveSource.cpp
@@ -0,0 +1,40 @@
+// Copyright Epic Games, Inc. All Rights Reserved.
+
+#include "ChaosMover/ChaosLayeredMoveSource.h"
+
+#include UE_INLINE_GENERATED_CPP_BY_NAME(ChaosLayeredMoveSource)
+
+
+void UChaosLayeredMoveSource::OnStart_Async(UMoverBlackboard* SimBlackboard, double CurrentSimTimeMs)
+{
+	if (FLayeredMoveBase* Move = LayeredMove.GetMutablePtr<FLayeredMoveBase>())
+	{
+		Move->StartMove_Async(SimBlackboard, CurrentSimTimeMs);
+	}
+}
+
+void UChaosLayeredMoveSource::GenerateMove_Async(const FMoverTickStartData& StartState, const FMoverTimeStep& TimeStep, UMoverBlackboard* SimBlackboard, FProposedMove& OutProposedMove)
+{
+	if (FLayeredMoveBase* Move = LayeredMove.GetMutablePtr<FLayeredMoveBase>())
+	{
+		Move->GenerateMove_Async(StartState, TimeStep, SimBlackboard, OutProposedMove);
+	}
+}
+
+bool UChaosLayeredMoveSource::IsFinished(double CurrentSimTimeMs) const
+{
+	if (const FLayeredMoveBase* Move = LayeredMove.GetPtr<FLayeredMoveBase>())
+	{
+		return Move->IsFinished(CurrentSimTimeMs);
+	}
+
+	return true;
+}
+
+void UChaosLayeredMoveSource::OnEnd_Async(UMoverBlackboard* SimBlackboard, double CurrentSimTimeMs)
+{
+	if (FLayeredMoveBase* Move = LayeredMove.GetMutablePtr<FLayeredMoveBase>())
+	{
+		Move->EndMove_Async(SimBlackboard, CurrentSimTimeMs);
+	}
+}
diff --git a/Engine/Plugins/Experimental/ChaosMover/Source/ChaosMover/Private/ChaosMoveExecutorBase.cpp b/Engine/Plugins/Experimental/ChaosMover/Source/ChaosMover/Private/ChaosMoveExecutorBase.cpp
new file mode 100644
--- /dev/null
+++ b/Engine/Plugins/Experimental/ChaosMover/Source/ChaosMover/Private/ChaosMoveExecutorBase.cpp
@@ -0,0 +1,5 @@
+// Copyright Epic Games, Inc. All Rights Reserved.
+
+#include "ChaosMover/ChaosMoveExecutorBase.h"
+
+#include UE_INLINE_GENERATED_CPP_BY_NAME(ChaosMoveExecutorBase)
diff --git a/Engine/Plugins/Experimental/ChaosMover/Source/ChaosMover/Private/ChaosMoverSourceBase.cpp b/Engine/Plugins/Experimental/ChaosMover/Source/ChaosMover/Private/ChaosMoverSourceBase.cpp
new file mode 100644
--- /dev/null
+++ b/Engine/Plugins/Experimental/ChaosMover/Source/ChaosMover/Private/ChaosMoverSourceBase.cpp
@@ -0,0 +1,5 @@
+// Copyright Epic Games, Inc. All Rights Reserved.
+
+#include "ChaosMover/ChaosMoverSourceBase.h"
+
+#include UE_INLINE_GENERATED_CPP_BY_NAME(ChaosMoverSourceBase)
diff --git a/Engine/Plugins/Experimental/ChaosMover/Source/ChaosMover/Private/Character/Executors/ChaosCharacterWalkingExecutor.cpp b/Engine/Plugins/Experimental/ChaosMover/Source/ChaosMover/Private/Character/Executors/ChaosCharacterWalkingExecutor.cpp
new file mode 100644
--- /dev/null
+++ b/Engine/Plugins/Experimental/ChaosMover/Source/ChaosMover/Private/Character/Executors/ChaosCharacterWalkingExecutor.cpp
@@ -0,0 +1,342 @@
+// Copyright Epic Games, Inc. All Rights Reserved.
+
+#include "ChaosMover/Character/Executors/ChaosCharacterWalkingExecutor.h"
+
+#include "Chaos/Character/CharacterGroundConstraintSettings.h"
+#include "ChaosMover/ChaosMoverLog.h"
+#include "ChaosMover/ChaosMoverSimulation.h"
+#include "ChaosMover/ChaosMoverSimulationTypes.h"
+#include "ChaosMover/Character/Settings/SharedChaosCharacterMovementSettings.h"
+#include "ChaosMover/Utilities/ChaosMoverQueryUtils.h"
+#include "Components/CapsuleComponent.h"
+#include "Components/SkeletalMeshComponent.h"
+#include "GameFramework/Character.h"
+#include "Math/UnitConversion.h"
+#include "MoverComponent.h"
+#include "MoveLibrary/FloorQueryUtils.h"
+#include "MoveLibrary/MovementUtils.h"
+#include "MoveLibrary/WaterMovementUtils.h"
+#include "MoverSimulationTypes.h"
+#include "Chaos/PhysicsObjectInternalInterface.h"
+
+#include UE_INLINE_GENERATED_CPP_BY_NAME(ChaosCharacterWalkingExecutor)
+
+
+void UChaosCharacterWalkingExecutor::OnModeRegistered(const FName ModeName)
+{
+	if (TargetHeightOverride.IsSet())
+	{
+		TargetHeight = TargetHeightOverride.GetValue();
+	}
+	else if (const ACharacter* CharacterOwner = GetTypedOuter<ACharacter>())
+	{
+		TargetHeight = -CharacterOwner->GetMesh()->GetRelativeLocation().Z;
+	}
+
+	if (QueryRadiusOverride.IsSet())
+	{
+		QueryRadius = QueryRadiusOverride.GetValue();
+	}
+	else if (const ACharacter* CharacterOwner = GetTypedOuter<ACharacter>())
+	{
+		if (const UCapsuleComponent* CapsuleComp = CharacterOwner->GetCapsuleComponent())
+		{
+			QueryRadius = FMath::Max(CapsuleComp->GetScaledCapsuleRadius() - 5.0f, 0.0f);
+		}
+	}
+
+	if (UMoverComponent* MoverComp = GetTypedOuter<UMoverComponent>())
+	{
+		TObjectPtr<const USharedChaosCharacterMovementSettings> SharedSettingsPtr = MoverComp->FindSharedSettings<USharedChaosCharacterMovementSettings>();
+		if (ensureMsgf(SharedSettingsPtr, TEXT("%s: Failed to find USharedChaosCharacterMovementSettings. Movement may not function properly."), *GetPathNameSafe(this)))
+		{
+			SharedSettingsWeakPtr = TWeakObjectPtr<const USharedChaosCharacterMovementSettings>(SharedSettingsPtr);
+		}
+	}
+}
+
+void UChaosCharacterWalkingExecutor::OnModeUnregistered()
+{
+	SharedSettingsWeakPtr.Reset();
+	MaxSpeedOverride.Reset();
+	AccelerationOverride.Reset();
+}
+
+void UChaosCharacterWalkingExecutor::ExecuteMove_Async(const FSimulationTickParams& Params, FMoverTickEndData& OutputState)
+{
+	QUICK_SCOPE_CYCLE_COUNTER(ChaosCharacterWalkingExecutor_ExecuteMove);
+
+	if (!Simulation)
+	{
+		UE_LOG(LogChaosMover, Warning, TEXT("UChaosCharacterWalkingExecutor (%s): No Simulation set. Check you have a Chaos Backend."), *GetName());
+		return;
+	}
+
+	const FChaosMoverSimulationDefaultInputs* DefaultSimInputs = Simulation->GetLocalSimInput().FindDataByType<FChaosMoverSimulationDefaultInputs>();
+	const FCharacterDefaultInputs* CharacterInputs = Params.StartState.InputCmd.InputCollection.FindDataByType<FCharacterDefaultInputs>();
+	if (!DefaultSimInputs || !CharacterInputs)
+	{
+		UE_LOG(LogChaosMover, Warning, TEXT("UChaosCharacterWalkingExecutor (%s): Requires FChaosMoverSimulationDefaultInputs and FCharacterDefaultInputs."), *GetName());
+		return;
+	}
+
+	TStrongObjectPtr<const USharedChaosCharacterMovementSettings> SharedSettingsStrongPtr = SharedSettingsWeakPtr.Pin();
+	if (!SharedSettingsStrongPtr)
+	{
+		UE_LOG(LogChaosMover, Warning, TEXT("UChaosCharacterWalkingExecutor (%s): Unable to read from shared settings (USharedChaosCharacterMovementSettings)."), *GetName());
+		return;
+	}
+	const USharedChaosCharacterMovementSettings* SharedSettings = SharedSettingsStrongPtr.Get();
+	check(SharedSettings);
+
+	FProposedMove ProposedMove = Params.ProposedMove;
+	const FMoverDefaultSyncState* StartingSyncState = Params.StartState.SyncState.SyncStateCollection.FindDataByType<FMoverDefaultSyncState>();
+	if (!StartingSyncState)
+	{
+		UE_LOG(LogChaosMover, Warning, TEXT("UChaosCharacterWalkingExecutor (%s): Requires FMoverDefaultSyncState in the sync state collection."), *GetName());
+		return;
+	}
+
+	FMoverDefaultSyncState& OutputSyncState = OutputState.SyncState.SyncStateCollection.FindOrAddMutableDataByType<FMoverDefaultSyncState>();
+	OutputSyncState = *StartingSyncState;
+
+	const float DeltaSeconds = Params.TimeStep.StepMs * 0.001f;
+	const FVector UpDirection = DefaultSimInputs->UpDir;
+	FVector GroundNormal = UpDirection;
+
+	FFloorCheckResult FloorResult;
+	UE::ChaosMover::FGroundDynamicsInfo GroundDynamicsInfo;
+	if (const UMoverBlackboard* SimBlackboard = Simulation->GetBlackboard())
+	{
+		SimBlackboard->TryGet(CommonBlackboard::LastFloorResult, FloorResult);
+		SimBlackboard->TryGet(UE::ChaosMover::Blackboard::GroundDynamicsInfo, GroundDynamicsInfo);
+	}
+
+	if (FloorResult.IsWalkableFloor())
+	{
+		GroundNormal = FloorResult.HitResult.ImpactNormal;
+
+		FVector TargetPosition = StartingSyncState->GetLocation_WorldSpace();
+
+		// Physics volumes always apply gravity along world Z (FVector::UpVector), not along UpDir.
+		const FVector PhysicsObjectGravityVelocity = DefaultSimInputs->PhysicsObjectGravity * FVector::UpVector * DeltaSeconds;
+
+		const FVector ProjectedVelocity = StartingSyncState->GetVelocity_WorldSpace() + DefaultSimInputs->Gravity * DeltaSeconds;
+		FVector TargetVelocity = ProjectedVelocity - PhysicsObjectGravityVelocity;
+
+		constexpr float ParallelCosThreshold = 0.999f;
+		const bool bNonVerticalVelocity = !FVector::Parallel(TargetVelocity.GetSafeNormal(), UpDirection, ParallelCosThreshold);
+		const bool bUseProposedMove = bNonVerticalVelocity || ProposedMove.bHasDirIntent;
+		bool bNormalVelocityIntent = false;
+
+		if (bUseProposedMove)
+		{
+			const FVector ProposedMovePlaneVelocity = ProposedMove.LinearVelocity - ProposedMove.LinearVelocity.ProjectOnToNormal(GroundNormal);
+
+			FVector ProposedNormalVelocity = ProposedMove.LinearVelocity - ProposedMovePlaneVelocity;
+			if (ProposedNormalVelocity.SizeSquared() > UE_KINDA_SMALL_NUMBER)
+			{
+				TargetVelocity += ProposedNormalVelocity - TargetVelocity.ProjectOnToNormal(GroundNormal);
+				bNormalVelocityIntent = true;
+			}
+
+			TargetPosition += ProposedMovePlaneVelocity * DeltaSeconds;
+		}
+
+		FVector ProjectedGroundVelocity = GroundDynamicsInfo.LinearVelocity;
+		if (GroundDynamicsInfo.bIsGravityEnabled)
+		{
+			ProjectedGroundVelocity += PhysicsObjectGravityVelocity;
+		}
+
+		const FVector ProjectedRelativeVelocity = TargetVelocity - ProjectedGroundVelocity;
+		const float ProjectedRelativeNormalVelocity = FloorResult.HitResult.ImpactNormal.Dot(TargetVelocity - ProjectedGroundVelocity);
+		const float ProjectedRelativeVerticalVelocity = GroundNormal.Dot(TargetVelocity - ProjectedGroundVelocity);
+		const float VerticalVelocityLimit = bNormalVelocityIntent ? 2.0f / DeltaSeconds : FMath::Abs(GroundNormal.Dot(DefaultSimInputs->Gravity) * DeltaSeconds);
+
+		bool bIsLiftingOffSurface = false;
+		if ((ProjectedRelativeNormalVelocity > VerticalVelocityLimit) && GroundDynamicsInfo.bIsMoving && (ProjectedRelativeVerticalVelocity > VerticalVelocityLimit))
+		{
+			bIsLiftingOffSurface = true;
+		}
+
+		const float InitialHeightAboveFloor = FloorResult.FloorDist - GetTargetHeight();
+		const float EndHeightAboveFloor = InitialHeightAboveFloor + ProjectedRelativeVerticalVelocity * DeltaSeconds;
+		const bool bIsSteppingDown = InitialHeightAboveFloor > UE_KINDA_SMALL_NUMBER;
+		const bool bIsWithinReach = EndHeightAboveFloor <= SharedSettings->MaxStepHeight;
+		const bool bIsSupported = bIsWithinReach && !bIsLiftingOffSurface;
+		const bool bNeedsVerticalVelocityToTarget = bIsSupported && bIsSteppingDown && (EndHeightAboveFloor > 0.0f) && !bIsLiftingOffSurface;
+		if (bNeedsVerticalVelocityToTarget)
+		{
+			TargetVelocity -= FractionalDownwardVelocityToTarget * (EndHeightAboveFloor / DeltaSeconds) * UpDirection;
+		}
+
+		if (!GroundDynamicsInfo.bIsMoving)
+		{
+			TargetPosition -= UpDirection * InitialHeightAboveFloor;
+		}
+
+		const FRotator TargetOrientation = UMovementUtils::ApplyAngularVelocityToRotator(StartingSyncState->GetOrientation_WorldSpace(), ProposedMove.AngularVelocityDegrees, DeltaSeconds);
+		FVector TargetAngularVelocity = StartingSyncState->GetAngularVelocityDegrees_WorldSpace();
+
+		if (ShouldCharacterRemainUpright() && bShouldApplyAngularVelocityToTarget)
+		{
+			if (DeltaSeconds > UE_SMALL_NUMBER)
+			{
+				const FQuat InitialQuat = StartingSyncState->GetOrientation_WorldSpace().Quaternion();
+				FQuat TgtQuat = TargetOrientation.Quaternion();
+				TgtQuat.EnforceShortestArcWith(InitialQuat);
+				const FQuat QuatRotation = TgtQuat * InitialQuat.Inverse();
+				const FVector AngularDisplacement = QuatRotation.ToRotationVector();
+				TargetAngularVelocity = FMath::RadiansToDegrees(AngularDisplacement / DeltaSeconds);
+			}
+		}
+
+		OutputSyncState.MoveDirectionIntent = ProposedMove.bHasDirIntent ? ProposedMove.DirectionIntent : FVector::ZeroVector;
+		OutputSyncState.SetTransforms_WorldSpace(
+			TargetPosition,
+			TargetOrientation,
+			TargetVelocity,
+			TargetAngularVelocity);
+	}
+
+	OutputState.MovementEndState.RemainingMs = 0.0f;
+	OutputState.MovementEndState.NextModeName = Params.StartState.SyncState.MovementMode;
+}
+
+// ---------------------------------------------------------------------------
+// IChaosCharacterMovementModeInterface
+// ---------------------------------------------------------------------------
+
+float UChaosCharacterWalkingExecutor::GetMaxWalkSlopeCosine() const
+{
+	if (TStrongObjectPtr<const USharedChaosCharacterMovementSettings> SharedSettingsStrongPtr = SharedSettingsWeakPtr.Pin())
+	{
+		return SharedSettingsStrongPtr->GetMaxWalkableSlopeCosine();
+	}
+
+	return 0.707f;
+}
+
+float UChaosCharacterWalkingExecutor::GetMaxSpeed() const
+{
+	if (MaxSpeedOverride.IsSet())
+	{
+		return MaxSpeedOverride.GetValue();
+	}
+
+	if (TStrongObjectPtr<const USharedChaosCharacterMovementSettings> SharedSettingsStrongPtr = SharedSettingsWeakPtr.Pin())
+	{
+		return SharedSettingsStrongPtr->MaxSpeed;
+	}
+
+	UE_LOG(LogChaosMover, Warning, TEXT("UChaosCharacterWalkingExecutor (%s): Unable to read MaxSpeed from shared settings."), *GetName());
+	return 0.0f;
+}
+
+void UChaosCharacterWalkingExecutor::OverrideMaxSpeed(float Value)
+{
+	MaxSpeedOverride = Value;
+}
+
+void UChaosCharacterWalkingExecutor::ClearMaxSpeedOverride()
+{
+	MaxSpeedOverride.Reset();
+}
+
+float UChaosCharacterWalkingExecutor::GetAcceleration() const
+{
+	if (AccelerationOverride.IsSet())
+	{
+		return AccelerationOverride.GetValue();
+	}
+
+	if (TStrongObjectPtr<const USharedChaosCharacterMovementSettings> SharedSettingsStrongPtr = SharedSettingsWeakPtr.Pin())
+	{
+		return SharedSettingsStrongPtr->Acceleration;
+	}
+
+	UE_LOG(LogChaosMover, Warning, TEXT("UChaosCharacterWalkingExecutor (%s): Unable to read Acceleration from shared settings."), *GetName());
+	return 0.0f;
+}
+
+void UChaosCharacterWalkingExecutor::OverrideAcceleration(float Value)
+{
+	AccelerationOverride = Value;
+}
+
+void UChaosCharacterWalkingExecutor::ClearAccelerationOverride()
+{
+	AccelerationOverride.Reset();
+}
+
+void UChaosCharacterWalkingExecutor::UpdateCurrentFloor(const FMoverTimeStep& TimeStep) const
+{
+	check(Simulation);
+
+	UMoverBlackboard* SimBlackboard = Simulation->GetBlackboard_Mutable();
+	if (!SimBlackboard)
+	{
+		return;
+	}
+
+	bool bFloorCheckSucceeded = false;
+
+	if (const FChaosMoverSimulationDefaultInputs* SimInputs = Simulation->GetLocalSimInput().FindDataByType<FChaosMoverSimulationDefaultInputs>())
+	{
+		Chaos::FReadPhysicsObjectInterface_Internal ReadInterface = Chaos::FPhysicsObjectInternalInterface::GetRead();
+		if (const Chaos::FPBDRigidParticleHandle* UpdatedParticle = ReadInterface.GetRigidParticle(SimInputs->PhysicsObject))
+		{
+			FFloorCheckResult FloorResult;
+			FWaterCheckResult WaterResult;
+
+			const float DeltaSeconds = TimeStep.StepMs * 0.001f;
+			const FVector VelocityWithGravity = UpdatedParticle->GetV() + UMovementUtils::ComputeVelocityFromGravity(SimInputs->Gravity, DeltaSeconds);
+
+			UE::ChaosMover::Utils::FFloorSweepParams SweepParams{
+				.ResponseParams = SimInputs->CollisionResponseParams,
+				.QueryParams = SimInputs->CollisionQueryParams,
+				.Location = UpdatedParticle->GetX(),
+				.DeltaPos = VelocityWithGravity * DeltaSeconds,
+				.UpDir = SimInputs->UpDir,
+				.World = SimInputs->World,
+				.QueryDistance = 1.5f * GetTargetHeight(),
+				.QueryRadius = FMath::Min(GetGroundQueryRadius(), FMath::Max(SimInputs->PawnCollisionRadius - 5.0f, 0.0f)),
+				.MaxWalkSlopeCosine = GetMaxWalkSlopeCosine(),
+				.TargetHeight = GetTargetHeight(),
+				.CollisionChannel = SimInputs->CollisionChannel
+			};
+
+			UE::ChaosMover::Utils::FloorSweep_Internal(SweepParams, FloorResult, WaterResult);
+
+			SimBlackboard->Set(CommonBlackboard::LastFloorResult, FloorResult);
+			SimBlackboard->Set(CommonBlackboard::LastWaterResult, WaterResult);
+			SimBlackboard->Set(UE::ChaosMover::Blackboard::GroundDynamicsInfo, UE::ChaosMover::FGroundDynamicsInfo(FloorResult));
+
+			bFloorCheckSucceeded = true;
+		}
+	}
+
+	if (!bFloorCheckSucceeded)
+	{
+		SimBlackboard->Invalidate(CommonBlackboard::LastFloorResult);
+		SimBlackboard->Invalidate(CommonBlackboard::LastWaterResult);
+		SimBlackboard->Invalidate(UE::ChaosMover::Blackboard::GroundDynamicsInfo);
+	}
+}
+
+// ---------------------------------------------------------------------------
+// IChaosCharacterConstraintMovementModeInterface
+// ---------------------------------------------------------------------------
+
+void UChaosCharacterWalkingExecutor::UpdateConstraintSettings(Chaos::FCharacterGroundConstraintSettings& ConstraintSettings) const
+{
+	ConstraintSettings.RadialForceLimit = FUnitConversion::Convert(RadialForceLimit, EUnit::Newtons, EUnit::KilogramCentimetersPerSecondSquared);
+	ConstraintSettings.TwistTorqueLimit = FUnitConversion::Convert(TwistTorqueLimit, EUnit::NewtonMeters, EUnit::KilogramCentimetersSquaredPerSecondSquared);
+	ConstraintSettings.SwingTorqueLimit = FUnitConversion::Convert(SwingTorqueLimit, EUnit::NewtonMeters, EUnit::KilogramCentimetersSquaredPerSecondSquared);
+	ConstraintSettings.TargetHeight = TargetHeight;
+	ConstraintSettings.FrictionForceLimit = FUnitConversion::Convert(FrictionForceLimit, EUnit::Newtons, EUnit::KilogramCentimetersPerSecondSquared);
+	ConstraintSettings.DampingFactor = GroundDamping;
+	ConstraintSettings.MotionTargetMassBias = FractionalGroundReaction;
+	ConstraintSettings.RadialForceMotionTargetScaling = FractionalRadialForceLimitScaling;
+}
diff --git a/Engine/Plugins/Experimental/ChaosMover/Source/ChaosMover/Private/Character/Modes/ChaosCharacterCompositeMovementMode.cpp b/Engine/Plugins/Experimental/ChaosMover/Source/ChaosMover/Private/Character/Modes/ChaosCharacterCompositeMovementMode.cpp
new file mode 100644
--- /dev/null
+++ b/Engine/Plugins/Experimental/ChaosMover/Source/ChaosMover/Private/Character/Modes/ChaosCharacterCompositeMovementMode.cpp
@@ -0,0 +1,158 @@
+// Copyright Epic Games, Inc. All Rights Reserved.
+
+#include "ChaosMover/Character/Modes/ChaosCharacterCompositeMovementMode.h"
+#include "Chaos/Character/CharacterGroundConstraintSettings.h"
+#include "ChaosMover/ChaosMoverLog.h"
+#include "Misc/DataValidation.h"
+
+#include UE_INLINE_GENERATED_CPP_BY_NAME(ChaosCharacterCompositeMovementMode)
+
+
+UChaosCharacterCompositeMovementMode::UChaosCharacterCompositeMovementMode(const FObjectInitializer& ObjectInitializer)
+	: Super(ObjectInitializer)
+{
+	SharedSettingsClasses.Add(USharedChaosCharacterMovementSettings::StaticClass());
+}
+
+void UChaosCharacterCompositeMovementMode::OnRegistered(const FName ModeName, const FMoverSimContext& SimContext)
+{
+	Super::OnRegistered(ModeName, SimContext);
+
+	CachedCharacterInterface = Cast<IChaosCharacterMovementModeInterface>(MoveExecutor);
+	CachedConstraintInterface = Cast<IChaosCharacterConstraintMovementModeInterface>(MoveExecutor);
+
+	ensureMsgf(CachedCharacterInterface,
+		TEXT("%s: MoveExecutor (%s) must implement IChaosCharacterMovementModeInterface to be used with UChaosCharacterCompositeMovementMode."),
+		*GetName(), *GetNameSafe(MoveExecutor));
+
+	ensureMsgf(CachedConstraintInterface,
+		TEXT("%s: MoveExecutor (%s) must implement IChaosCharacterConstraintMovementModeInterface to be used with UChaosCharacterCompositeMovementMode."),
+		*GetName(), *GetNameSafe(MoveExecutor));
+}
+
+void UChaosCharacterCompositeMovementMode::OnUnregistered(const FMoverSimContext& SimContext)
+{
+	CachedCharacterInterface = nullptr;
+	CachedConstraintInterface = nullptr;
+
+	Super::OnUnregistered(SimContext);
+}
+
+#if WITH_EDITOR
+EDataValidationResult UChaosCharacterCompositeMovementMode::IsDataValid(FDataValidationContext& Context) const
+{
+	EDataValidationResult Result = Super::IsDataValid(Context);
+
+	if (MoveExecutor)
+	{
+		if (!Cast<IChaosCharacterMovementModeInterface>(MoveExecutor))
+		{
+			Context.AddError(FText::FromString(FString::Printf(
+				TEXT("MoveExecutor '%s' must implement IChaosCharacterMovementModeInterface."),
+				*GetNameSafe(MoveExecutor))));
+			Result = EDataValidationResult::Invalid;
+		}
+
+		if (!Cast<IChaosCharacterConstraintMovementModeInterface>(MoveExecutor))
+		{
+			Context.AddError(FText::FromString(FString::Printf(
+				TEXT("MoveExecutor '%s' must implement IChaosCharacterConstraintMovementModeInterface."),
+				*GetNameSafe(MoveExecutor))));
+			Result = EDataValidationResult::Invalid;
+		}
+	}
+
+	return Result;
+}
+#endif
+
+// ---------------------------------------------------------------------------
+// IChaosCharacterMovementModeInterface — delegated to MoveExecutor
+// ---------------------------------------------------------------------------
+
+float UChaosCharacterCompositeMovementMode::GetTargetHeight() const
+{
+	return CachedCharacterInterface ? CachedCharacterInterface->GetTargetHeight() : 0.f;
+}
+
+float UChaosCharacterCompositeMovementMode::GetGroundQueryRadius() const
+{
+	return CachedCharacterInterface ? CachedCharacterInterface->GetGroundQueryRadius() : 0.f;
+}
+
+float UChaosCharacterCompositeMovementMode::GetMaxWalkSlopeCosine() const
+{
+	return CachedCharacterInterface ? CachedCharacterInterface->GetMaxWalkSlopeCosine() : 0.f;
+}
+
+bool UChaosCharacterCompositeMovementMode::ShouldCharacterRemainUpright() const
+{
+	return CachedCharacterInterface ? CachedCharacterInterface->ShouldCharacterRemainUpright() : true;
+}
+
+float UChaosCharacterCompositeMovementMode::GetMaxSpeed() const
+{
+	return CachedCharacterInterface ? CachedCharacterInterface->GetMaxSpeed() : 0.f;
+}
+
+void UChaosCharacterCompositeMovementMode::OverrideMaxSpeed(float Value)
+{
+	if (CachedCharacterInterface)
+	{
+		CachedCharacterInterface->OverrideMaxSpeed(Value);
+	}
+}
+
+void UChaosCharacterCompositeMovementMode::ClearMaxSpeedOverride()
+{
+	if (CachedCharacterInterface)
+	{
+		CachedCharacterInterface->ClearMaxSpeedOverride();
+	}
+}
+
+float UChaosCharacterCompositeMovementMode::GetAcceleration() const
+{
+	return CachedCharacterInterface ? CachedCharacterInterface->GetAcceleration() : 0.f;
+}
+
+void UChaosCharacterCompositeMovementMode::OverrideAcceleration(float Value)
+{
+	if (CachedCharacterInterface)
+	{
+		CachedCharacterInterface->OverrideAcceleration(Value);
+	}
+}
+
+void UChaosCharacterCompositeMovementMode::ClearAccelerationOverride()
+{
+	if (CachedCharacterInterface)
+	{
+		CachedCharacterInterface->ClearAccelerationOverride();
+	}
+}
+
+void UChaosCharacterCompositeMovementMode::UpdateCurrentFloor(const FMoverTimeStep& TimeStep) const
+{
+	if (CachedCharacterInterface)
+	{
+		CachedCharacterInterface->UpdateCurrentFloor(TimeStep);
+	}
+}
+
+// ---------------------------------------------------------------------------
+// IChaosCharacterConstraintMovementModeInterface — delegated to MoveExecutor
+// ---------------------------------------------------------------------------
+
+bool UChaosCharacterCompositeMovementMode::ShouldEnableConstraint() const
+{
+	return CachedConstraintInterface ? CachedConstraintInterface->ShouldEnableConstraint() : false;
+}
+
+void UChaosCharacterCompositeMovementMode::UpdateConstraintSettings(Chaos::FCharacterGroundConstraintSettings& ConstraintSettings) const
+{
+	if (CachedConstraintInterface)
+	{
+		CachedConstraintInterface->UpdateConstraintSettings(ConstraintSettings);
+	}
+}
diff --git a/Engine/Plugins/Experimental/ChaosMover/Source/ChaosMover/Public/ChaosMover/ChaosCompositeMovementMode.h b/Engine/Plugins/Experimental/ChaosMover/Source/ChaosMover/Public/ChaosMover/ChaosCompositeMovementMode.h
new file mode 100644
--- /dev/null
+++ b/Engine/Plugins/Experimental/ChaosMover/Source/ChaosMover/Public/ChaosMover/ChaosCompositeMovementMode.h
@@ -0,0 +1,63 @@
+// Copyright Epic Games, Inc. All Rights Reserved.
+
+#pragma once
+
+#include "CoreMinimal.h"
+#include "ChaosMover/ChaosMovementMode.h"
+#include "ChaosMover/ChaosMoverSourceBase.h"
+#include "ChaosMover/ChaosMoveExecutorBase.h"
+#include "ChaosCompositeMovementMode.generated.h"
+
+#define UE_API CHAOSMOVER_API
+
+/**
+ * UChaosCompositeMovementMode: a ChaosMover movement mode that composes a
+ * UChaosMoverSourceBase (what velocity to produce) with a UChaosMoveExecutorBase
+ * (how to apply it to the physics simulation). This separates move generation from
+ * physics application so that both can be configured and swapped independently.
+ *
+ * All simulation methods are async-safe — no AActor or UActorComponent access.
+ *
+ * When MoveSource::IsFinished() returns true, the mode automatically transitions to
+ * NextModeName (if set). The Transitions array can be used for additional conditions.
+ *
+ * Derive from UChaosCharacterCompositeMovementMode instead if the executor
+ * needs to participate in the Chaos character ground constraint system.
+ */
+UCLASS(MinimalAPI, Abstract, Blueprintable, BlueprintType)
+class UChaosCompositeMovementMode : public UChaosMovementMode
+{
+	GENERATED_BODY()
+
+public:
+	UE_API UChaosCompositeMovementMode(const FObjectInitializer& ObjectInitializer);
+
+	/** The source that produces a FProposedMove each tick. */
+	UPROPERTY(EditAnywhere, Instanced, Category = Mover)
+	TObjectPtr<UChaosMoverSourceBase> MoveSource;
+
+	/** Applies the proposed move to the Chaos physics simulation state each tick. */
+	UPROPERTY(EditAnywhere, Instanced, Category = Mover)
+	TObjectPtr<UChaosMoveExecutorBase> MoveExecutor;
+
+	/** Mode to transition to when MoveSource::IsFinished() returns true. NAME_None = no auto-transition. */
+	UPROPERTY(EditAnywhere, BlueprintReadWrite, Category = Mover)
+	FName NextModeName;
+
+	UE_API virtual void OnRegistered(const FName ModeName, const FMoverSimContext& SimContext) override;
+	UE_API virtual void OnUnregistered(const FMoverSimContext& SimContext) override;
+
+	UE_API virtual void Activate(const FMoverEventContext& Context, FName PrevModeName, const FMoverSimContext& SimContext, const FMoverTickStartData& StartState, FMoverSyncState* OutSyncState, FMoverAuxStateContext* OutAuxState) override;
+	UE_API virtual void Deactivate(const FMoverEventContext& Context, FName InNextModeName, const FMoverSimContext& SimContext) override;
+
+	UE_API virtual void GenerateMove_Implementation(const FMoverSimContext& SimContext, const FMoverTickStartData& StartState, const FMoverTimeStep& TimeStep, FProposedMove& OutProposedMove) const override;
+	UE_API virtual void SimulationTick_Implementation(const FSimulationTickParams& Params, FMoverTickEndData& OutputState) override;
+
+private:
+	// Set when MoveSource->OnEnd_Async has been called, cleared on Activate.
+	// Prevents OnEnd_Async from being called more than once per activation when
+	// IsFinished() returns true but no mode transition is configured.
+	bool bMoveSourceEnded = false;
+};
+
+#undef UE_API
diff --git a/Engine/Plugins/Experimental/ChaosMover/Source/ChaosMover/Public/ChaosMover/ChaosLayeredMoveSource.h b/Engine/Plugins/Experimental/ChaosMover/Source/ChaosMover/Public/ChaosMover/ChaosLayeredMoveSource.h
new file mode 100644
--- /dev/null
+++ b/Engine/Plugins/Experimental/ChaosMover/Source/ChaosMover/Public/ChaosMover/ChaosLayeredMoveSource.h
@@ -0,0 +1,38 @@
+// Copyright Epic Games, Inc. All Rights Reserved.
+
+#pragma once
+
+#include "CoreMinimal.h"
+#include "ChaosMover/ChaosMoverSourceBase.h"
+#include "LayeredMove.h"
+#include "StructUtils/InstancedStruct.h"
+#include "ChaosLayeredMoveSource.generated.h"
+
+#define UE_API CHAOSMOVER_API
+
+/**
+ * UChaosLayeredMoveSource: a UChaosMoverSourceBase that drives move generation
+ * from a single FLayeredMoveBase instance stored as an instanced struct property.
+ *
+ * All four source lifecycle methods (OnStart_Async, GenerateMove_Async, IsFinished,
+ * OnEnd_Async) are fully delegated to the underlying FLayeredMoveBase async API,
+ * so subclasses of FLayeredMoveBase can be configured and swapped in the editor
+ * without subclassing this object.
+ */
+UCLASS(MinimalAPI, Blueprintable, BlueprintType, EditInlineNew, DefaultToInstanced)
+class UChaosLayeredMoveSource : public UChaosMoverSourceBase
+{
+	GENERATED_BODY()
+
+public:
+	/** The layered move that drives velocity generation. */
+	UPROPERTY(EditAnywhere, Category = Mover)
+	TInstancedStruct<FLayeredMoveBase> LayeredMove;
+
+	UE_API virtual void OnStart_Async(UMoverBlackboard* SimBlackboard, double CurrentSimTimeMs) override;
+	UE_API virtual void GenerateMove_Async(const FMoverTickStartData& StartState, const FMoverTimeStep& TimeStep, UMoverBlackboard* SimBlackboard, FProposedMove& OutProposedMove) override;
+	UE_API virtual bool IsFinished(double CurrentSimTimeMs) const override;
+	UE_API virtual void OnEnd_Async(UMoverBlackboard* SimBlackboard, double CurrentSimTimeMs) override;
+};
+
+#undef UE_API
diff --git a/Engine/Plugins/Experimental/ChaosMover/Source/ChaosMover/Public/ChaosMover/ChaosMoveExecutorBase.h b/Engine/Plugins/Experimental/ChaosMover/Source/ChaosMover/Public/ChaosMover/ChaosMoveExecutorBase.h
new file mode 100644
--- /dev/null
+++ b/Engine/Plugins/Experimental/ChaosMover/Source/ChaosMover/Public/ChaosMover/ChaosMoveExecutorBase.h
@@ -0,0 +1,72 @@
+// Copyright Epic Games, Inc. All Rights Reserved.
+
+#pragma once
+
+#include "CoreMinimal.h"
+#include "ChaosMoveExecutorBase.generated.h"
+
+#define UE_API CHAOSMOVER_API
+
+class UChaosMoverSimulation;
+struct FSimulationTickParams;
+struct FMoverTickEndData;
+
+/**
+ * Abstract base class for move executors in ChaosMover.
+ *
+ * A move executor is responsible for applying a FProposedMove to the physics simulation
+ * state — computing target delta positions, updating the sync state, and handling any
+ * surface interaction. It decouples the "how to apply movement" concern from the mode
+ * hierarchy, allowing physics behaviors (walking, free-move, etc.) to be composed
+ * independently of move generation.
+ *
+ * All methods must be safe to call from the physics simulation thread — no access to
+ * AActor or UActorComponent is permitted.
+ *
+ * Concrete subclasses that support the Chaos character ground constraint should
+ * additionally implement IChaosCharacterMovementModeInterface and
+ * IChaosCharacterConstraintMovementModeInterface so that
+ * UChaosCharacterCompositeMovementMode can delegate those interfaces to them.
+ */
+UCLASS(Abstract, MinimalAPI, EditInlineNew, DefaultToInstanced)
+class UChaosMoveExecutorBase : public UObject
+{
+	GENERATED_BODY()
+
+public:
+	/**
+	 * Called by the owning mode during its OnRegistered, on the game thread.
+	 * Subclasses may override to perform initialization that requires game-thread
+	 * access (e.g. reading character mesh offset, finding shared settings).
+	 */
+	UE_API virtual void OnModeRegistered(const FName ModeName) {}
+
+	/**
+	 * Called by the owning mode during its OnUnregistered, on the game thread.
+	 * Subclasses may override to release any references acquired in OnModeRegistered.
+	 */
+	UE_API virtual void OnModeUnregistered() {}
+
+	/**
+	 * Applies Params.ProposedMove to the simulation output state. Thread-safe —
+	 * no AActor or UActorComponent access.
+	 */
+	UE_API virtual void ExecuteMove_Async(
+		const FSimulationTickParams& Params,
+		FMoverTickEndData& OutputState) {}
+
+	/**
+	 * Propagates the simulation reference so that the executor can perform physics
+	 * queries (e.g. floor sweeps) when needed. Called by the owning mode during
+	 * OnRegistered and cleared during OnUnregistered.
+	 */
+	UE_API virtual void SetSimulation(UChaosMoverSimulation* InSimulation)
+	{
+		Simulation = InSimulation;
+	}
+
+protected:
+	TObjectPtr<UChaosMoverSimulation> Simulation;
+};
+
+#undef UE_API
diff --git a/Engine/Plugins/Experimental/ChaosMover/Source/ChaosMover/Public/ChaosMover/ChaosMoverSourceBase.h b/Engine/Plugins/Experimental/ChaosMover/Source/ChaosMover/Public/ChaosMover/ChaosMoverSourceBase.h
new file mode 100644
--- /dev/null
+++ b/Engine/Plugins/Experimental/ChaosMover/Source/ChaosMover/Public/ChaosMover/ChaosMoverSourceBase.h
@@ -0,0 +1,51 @@
+// Copyright Epic Games, Inc. All Rights Reserved.
+
+#pragma once
+
+#include "CoreMinimal.h"
+#include "ChaosMoverSourceBase.generated.h"
+
+#define UE_API CHAOSMOVER_API
+
+class UMoverBlackboard;
+struct FMoverTickStartData;
+struct FMoverTimeStep;
+struct FProposedMove;
+
+/**
+ * Abstract base class for move sources.
+ *
+ * A move source is responsible for producing a FProposedMove from simulation state. It
+ * decouples the "what velocity to generate" concern from the movement mode, allowing the
+ * same generation logic to be reused across different modes or physics backends.
+ *
+ * All methods must be safe to call from the physics simulation thread — no access to
+ * AActor or UActorComponent is permitted.
+ */
+UCLASS(Abstract, MinimalAPI, EditInlineNew, DefaultToInstanced)
+class UChaosMoverSourceBase : public UObject
+{
+	GENERATED_BODY()
+
+public:
+	/** Called once when the owning mode is activated. Async-safe. */
+	UE_API virtual void OnStart_Async(UMoverBlackboard* SimBlackboard, double CurrentSimTimeMs) {}
+
+	/**
+	 * Produces a proposed move from the current simulation state. Called during the mode's
+	 * GenerateMove phase. Async-safe — no AActor or UActorComponent access.
+	 */
+	UE_API virtual void GenerateMove_Async(
+		const FMoverTickStartData& StartState,
+		const FMoverTimeStep& TimeStep,
+		UMoverBlackboard* SimBlackboard,
+		FProposedMove& OutProposedMove) {}
+
+	/** Returns true when this source has naturally expired (e.g. a duration-limited move has elapsed). */
+	UE_API virtual bool IsFinished(double CurrentSimTimeMs) const { return false; }
+
+	/** Called when the owning mode is deactivated or when IsFinished() returns true. Async-safe. */
+	UE_API virtual void OnEnd_Async(UMoverBlackboard* SimBlackboard, double CurrentSimTimeMs) {}
+};
+
+#undef UE_API
diff --git a/Engine/Plugins/Experimental/ChaosMover/Source/ChaosMover/Public/ChaosMover/Character/Executors/ChaosCharacterWalkingExecutor.h b/Engine/Plugins/Experimental/ChaosMover/Source/ChaosMover/Public/ChaosMover/Character/Executors/ChaosCharacterWalkingExecutor.h
new file mode 100644
--- /dev/null
+++ b/Engine/Plugins/Experimental/ChaosMover/Source/ChaosMover/Public/ChaosMover/Character/Executors/ChaosCharacterWalkingExecutor.h
@@ -0,0 +1,127 @@
+// Copyright Epic Games, Inc. All Rights Reserved.
+
+#pragma once
+
+#include "CoreMinimal.h"
+#include "ChaosMover/ChaosMoveExecutorBase.h"
+#include "ChaosMover/ChaosMoverSimulationTypes.h"
+#include "ChaosCharacterWalkingExecutor.generated.h"
+
+#define UE_API CHAOSMOVER_API
+
+class USharedChaosCharacterMovementSettings;
+
+/**
+ * UChaosCharacterWalkingExecutor: a concrete UChaosMoveExecutorBase that applies
+ * walking physics to the Chaos character ground constraint simulation.
+ *
+ * This executor also implements IChaosCharacterMovementModeInterface and
+ * IChaosCharacterConstraintMovementModeInterface, making it compatible with
+ * UChaosCharacterCompositeMovementMode which delegates those interfaces to its executor.
+ *
+ * The physics application logic (ExecuteMove_Async) is equivalent to
+ * UChaosWalkingMode::SimulationTick_Implementation. Move generation is intentionally
+ * left to the paired UChaosMoverSourceBase — use UChaosLayeredMoveSource or a custom
+ * subclass to drive the proposed velocity each tick.
+ *
+ * OnModeRegistered resolves TargetHeight and QueryRadius from the owning character
+ * and locates the shared movement settings — all on the game thread before simulation
+ * begins.
+ */
+UCLASS(MinimalAPI, Blueprintable, BlueprintType, EditInlineNew, DefaultToInstanced)
+class UChaosCharacterWalkingExecutor
+	: public UChaosMoveExecutorBase
+	, public IChaosCharacterMovementModeInterface
+	, public IChaosCharacterConstraintMovementModeInterface
+{
+	GENERATED_BODY()
+
+public:
+	// Maximum force the character can apply to reach the motion target
+	UPROPERTY(EditAnywhere, Category = "Constraint Settings", meta = (ClampMin = "0", UIMin = "0", ForceUnits = "Newtons"))
+	float RadialForceLimit = 2000.0f;
+
+	// Maximum torque the character can apply to rotate in air about the vertical axis
+	UPROPERTY(EditAnywhere, Category = "Constraint Settings", meta = (ClampMin = "0", UIMin = "0", ForceUnits = "NewtonMeters"))
+	float TwistTorqueLimit = 1500.0f;
+
+	// Maximum torque the character can apply to remain upright
+	UPROPERTY(EditAnywhere, Category = "Constraint Settings", meta = (ClampMin = "0", UIMin = "0", ForceUnits = "NewtonMeters"))
+	float SwingTorqueLimit = 3000.0f;
+
+	// Damping factor to control the softness of the interaction between the character and the ground.
+	// Set to 0 for no damping and 1 for maximum damping.
+	UPROPERTY(EditAnywhere, Category = "Constraint Settings", meta = (ClampMin = "0", UIMin = "0", ClampMax = "1", UIMax = "1"))
+	float GroundDamping = 0.0f;
+
+	// Maximum force the character can apply to hold in place while standing on an unwalkable incline
+	UPROPERTY(EditAnywhere, Category = "Constraint Settings", meta = (ClampMin = "0", UIMin = "0", ForceUnits = "Newtons"))
+	float FrictionForceLimit = 100.0f;
+
+	// Scaling applied to the radial force limit to raise the limit to always allow the character to
+	// reach the motion target. 1 = scale as needed, 0 = no scaling.
+	UPROPERTY(EditAnywhere, Category = "Constraint Settings", meta = (ClampMin = "0", UIMin = "0", ClampMax = "1", UIMax = "1"))
+	float FractionalRadialForceLimitScaling = 1.0f;
+
+	// Controls the reaction force applied to the ground in the ground plane when the character is moving.
+	// 1 = full reaction force, 0 = normal force only.
+	UPROPERTY(EditAnywhere, Category = "Constraint Settings", meta = (ClampMin = "0", UIMin = "0", ClampMax = "1", UIMax = "1"))
+	float FractionalGroundReaction = 1.0f;
+
+	// Controls whether the character capsule is forced to remain upright
+	UPROPERTY(EditAnywhere, Category = "Constraint Settings")
+	bool bShouldCharacterRemainUpright = true;
+
+	/**
+	 * Optional override for target height (desired distance from capsule center to floor).
+	 * If unset, derived automatically from the owning character's skeletal mesh Z offset.
+	 */
+	UPROPERTY(EditAnywhere, Category = "Constraint Settings", meta = (ClampMin = "0", UIMin = "0", ForceUnits = "cm"))
+	TOptional<float> TargetHeightOverride;
+
+	// Optional override for the ground query radius. If unset, derived from capsule radius.
+	UPROPERTY(EditAnywhere, Category = "Query Settings", meta = (ClampMin = "0", UIMin = "0", ForceUnits = "cm"))
+	TOptional<float> QueryRadiusOverride;
+
+	// Controls how much downward velocity is applied to keep the character rooted to the ground
+	// when within MaxStepHeight of the surface.
+	UPROPERTY(EditAnywhere, Category = "Movement Settings", meta = (ClampMin = "0", UIMin = "0", ClampMax = "1", UIMax = "1"))
+	float FractionalDownwardVelocityToTarget = 1.0f;
+
+	// Whether the character can apply angular velocity in response to stick input
+	UPROPERTY(EditAnywhere, BlueprintReadWrite, Category = "Movement Settings")
+	bool bShouldApplyAngularVelocityToTarget = true;
+
+	// UChaosMoveExecutorBase overrides
+	UE_API virtual void OnModeRegistered(const FName ModeName) override;
+	UE_API virtual void OnModeUnregistered() override;
+	UE_API virtual void ExecuteMove_Async(const FSimulationTickParams& Params, FMoverTickEndData& OutputState) override;
+
+	// IChaosCharacterMovementModeInterface — implemented directly on this executor
+	virtual float GetTargetHeight() const override { return TargetHeight; }
+	virtual float GetGroundQueryRadius() const override { return QueryRadius; }
+	UE_API virtual float GetMaxWalkSlopeCosine() const override;
+	virtual bool ShouldCharacterRemainUpright() const override { return bShouldCharacterRemainUpright; }
+	UE_API virtual float GetMaxSpeed() const override;
+	UE_API virtual void OverrideMaxSpeed(float Value) override;
+	UE_API virtual void ClearMaxSpeedOverride() override;
+	UE_API virtual float GetAcceleration() const override;
+	UE_API virtual void OverrideAcceleration(float Value) override;
+	UE_API virtual void ClearAccelerationOverride() override;
+	UE_API virtual void UpdateCurrentFloor(const FMoverTimeStep& TimeStep) const override;
+
+	// IChaosCharacterConstraintMovementModeInterface — implemented directly on this executor
+	virtual bool ShouldEnableConstraint() const override { return true; }
+	UE_API virtual void UpdateConstraintSettings(Chaos::FCharacterGroundConstraintSettings& ConstraintSettings) const override;
+
+private:
+	// Resolved on the game thread in OnModeRegistered
+	float TargetHeight = 95.0f;
+	float QueryRadius = 30.0f;
+	TWeakObjectPtr<const USharedChaosCharacterMovementSettings> SharedSettingsWeakPtr;
+
+	TOptional<float> MaxSpeedOverride;
+	TOptional<float> AccelerationOverride;
+};
+
+#undef UE_API
diff --git a/Engine/Plugins/Experimental/ChaosMover/Source/ChaosMover/Public/ChaosMover/Character/Modes/ChaosCharacterCompositeMovementMode.h b/Engine/Plugins/Experimental/ChaosMover/Source/ChaosMover/Public/ChaosMover/Character/Modes/ChaosCharacterCompositeMovementMode.h
new file mode 100644
--- /dev/null
+++ b/Engine/Plugins/Experimental/ChaosMover/Source/ChaosMover/Public/ChaosMover/Character/Modes/ChaosCharacterCompositeMovementMode.h
@@ -0,0 +1,72 @@
+// Copyright Epic Games, Inc. All Rights Reserved.
+
+#pragma once
+
+#include "CoreMinimal.h"
+#include "ChaosMover/ChaosCompositeMovementMode.h"
+#include "ChaosMover/ChaosMoverSimulationTypes.h"
+#include "ChaosMover/Character/Settings/SharedChaosCharacterMovementSettings.h"
+#include "ChaosCharacterCompositeMovementMode.generated.h"
+
+#define UE_API CHAOSMOVER_API
+
+/**
+ * UChaosCharacterCompositeMovementMode: extends UChaosCompositeMovementMode
+ * with participation in the Chaos character ground constraint system.
+ *
+ * This class implements IChaosCharacterMovementModeInterface and
+ * IChaosCharacterConstraintMovementModeInterface entirely by delegating to the assigned
+ * MoveExecutor. The executor is required to implement both interfaces — this is validated
+ * at registration time via ensureMsgf and also enforced during editor validation via
+ * IsDataValid.
+ *
+ * This design means all character constraint parameters (target height, query radius,
+ * force limits, floor query logic, etc.) are owned by the executor, and the mode itself
+ * carries no character-specific state.
+ */
+UCLASS(MinimalAPI, Blueprintable, BlueprintType, EditInlineNew, DefaultToInstanced)
+class UChaosCharacterCompositeMovementMode
+	: public UChaosCompositeMovementMode
+	, public IChaosCharacterMovementModeInterface
+	, public IChaosCharacterConstraintMovementModeInterface
+{
+	GENERATED_BODY()
+
+public:
+	UE_API UChaosCharacterCompositeMovementMode(const FObjectInitializer& ObjectInitializer);
+
+	// UBaseMovementMode overrides — cache and validate executor interface pointers
+	UE_API virtual void OnRegistered(const FName ModeName, const FMoverSimContext& SimContext) override;
+	UE_API virtual void OnUnregistered(const FMoverSimContext& SimContext) override;
+
+#if WITH_EDITOR
+	UE_API virtual EDataValidationResult IsDataValid(class FDataValidationContext& Context) const override;
+#endif
+
+	// IChaosCharacterMovementModeInterface — fully delegated to MoveExecutor
+	virtual float GetTargetHeight() const override;
+	virtual float GetGroundQueryRadius() const override;
+	virtual float GetMaxWalkSlopeCosine() const override;
+	virtual bool ShouldCharacterRemainUpright() const override;
+	virtual float GetMaxSpeed() const override;
+	virtual void OverrideMaxSpeed(float Value) override;
+	virtual void ClearMaxSpeedOverride() override;
+	virtual float GetAcceleration() const override;
+	virtual void OverrideAcceleration(float Value) override;
+	virtual void ClearAccelerationOverride() override;
+	virtual void UpdateCurrentFloor(const FMoverTimeStep& TimeStep) const override;
+
+	// IChaosCharacterConstraintMovementModeInterface — fully delegated to MoveExecutor
+	// GetTargetHeight() is shared with IChaosCharacterMovementModeInterface above
+	virtual bool ShouldEnableConstraint() const override;
+	virtual void UpdateConstraintSettings(Chaos::FCharacterGroundConstraintSettings& ConstraintSettings) const override;
+
+private:
+	// Cached interface pointers into MoveExecutor, populated in OnRegistered.
+	// Raw pointers are safe here: MoveExecutor is a UPROPERTY-tracked UObject owned by
+	// this mode, guaranteed to outlive these pointers within our registration lifetime.
+	IChaosCharacterMovementModeInterface* CachedCharacterInterface = nullptr;
+	IChaosCharacterConstraintMovementModeInterface* CachedConstraintInterface = nullptr;
+};
+
+#undef UE_API
