ChaosMover: Fix for crouching bugs
- ChaosMover was not respecting np2.Resim.RecordStatePostSolve and was always recording state pre solve. This led to a one frame offset, client would always apply server state one frame late
- Moving stance modifier event out of processed events, filtering is being moved to GT. Eventually ProcessedEvents will go away but a generic solution is not ready yet, so doing it piecemeal here to fix crouch for 40.30
- Compacting FChaosStanceModifier::UpdateStanc a bit for readibility and refining the early out when the stance doesn't change: that should only ever happen on rollback, i.e. on the first resim frame, and explaining how that comes about
- Propagate crouching intent on every frame, so resim frames can reconstruct intent. Previously, if we ever rolled back to just after the intent was received on the PT, there would be no way to know intent and resim correctly, input would get debounced by rollbacks

#Synced-CL 52217694

diff --git a/Engine/Plugins/Experimental/ChaosMover/Source/ChaosMover/Private/ChaosMoverSimulation.cpp b/Engine/Plugins/Experimental/ChaosMover/Source/ChaosMover/Private/ChaosMoverSimulation.cpp
--- a/Engine/Plugins/Experimental/ChaosMover/Source/ChaosMover/Private/ChaosMoverSimulation.cpp
+++ b/Engine/Plugins/Experimental/ChaosMover/Source/ChaosMover/Private/ChaosMoverSimulation.cpp
@@ -106,7 +106,7 @@ void UChaosMoverSimulation::ApplyNetStateData(const FMoverSyncState& InNetSyncSt
 
 void UChaosMoverSimulation::BuildNetStateData(FMoverSyncState& OutNetSyncState) const
 {
-	OutNetSyncState = NetSyncState;
+	OutNetSyncState = CurrentSyncState;
 }
 
 void UChaosMoverSimulation::Init(const FInitParams& InitParams)
@@ -449,10 +449,6 @@ void UChaosMoverSimulation::ProcessInputs(int32 PhysicsStep, const FMoverTimeSte
 		OnSimulationRollback(TimeStep, PrevSyncState);
 		bSyncStateOverridden = false;
 	}
-	else
-	{
-		NetSyncState = CurrentSyncState;
-	}
 }
 
 void UChaosMoverSimulation::SimulationTick(const FMoverSimContext& SimContext, const FMoverTimeStep& TimeStep, const UE::ChaosMover::FSimulationInputData& InputData, UE::ChaosMover::FSimulationOutputData& OutputData)
@@ -474,6 +470,7 @@ void UChaosMoverSimulation::SimulationTick(const FMoverSimContext& SimContext, c
 
 		RollbackBlackboard->EndSimulationFrame();
 	}
+
 }
 
 bool UChaosMoverSimulation::ShouldResim() const
@@ -1130,6 +1127,11 @@ void UChaosMoverSimulation::UpdateEvents(const FMoverTimeStep& TimeStep, UE::Cha
 					// Remove events that have already been processed in the main update
 
 					Events.RemoveAll([this](const TSharedPtr<FMoverSimulationEventData> Event) {
+						// Stance events are handled by batch accumulation on the game thread - skip dedup
+						if (Event->CastTo<FStanceModifiedEventData>())
+						{
+							return false;
+						}
 						const int32 AllowedDiff = UE::ChaosMover::CVars::FrameDifferenceLeniencyForEventTimeComparison;
 						for (TSharedPtr<FMoverSimulationEventData>& ProcessedEvent : ProcessedEvents)
 						{
diff --git a/Engine/Plugins/Experimental/ChaosMover/Source/ChaosMover/Private/Character/ChaosCharacterMoverComponent.cpp b/Engine/Plugins/Experimental/ChaosMover/Source/ChaosMover/Private/Character/ChaosCharacterMoverComponent.cpp
--- a/Engine/Plugins/Experimental/ChaosMover/Source/ChaosMover/Private/Character/ChaosCharacterMoverComponent.cpp
+++ b/Engine/Plugins/Experimental/ChaosMover/Source/ChaosMover/Private/Character/ChaosCharacterMoverComponent.cpp
@@ -123,6 +123,25 @@ void UChaosCharacterMoverComponent::OnMoverPreSimulationTick(const FMoverTimeSte
 	// because the stance modifier and jump handling of the base is not async physics friendly
 }
 
+void UChaosCharacterMoverComponent::DispatchSimulationEvents(const UE::Mover::FSimulationOutputData& OutputData)
+{
+	Super::DispatchSimulationEvents(OutputData);
+
+	// Flush the stance batch accumulated during this frame's event dispatch.
+	// If rollback/resim produced transient stance transitions (e.g. uncrouch then recrouch),
+	// only the net final stance is dispatched, suppressing spurious visual/animation callbacks.
+	if (PendingBatchFinalStance.IsSet())
+	{
+		const EStanceMode FinalStance = PendingBatchFinalStance.GetValue();
+		if (FinalStance != LastDispatchedStance)
+		{
+			OnStanceChanged.Broadcast(LastDispatchedStance, FinalStance);
+			LastDispatchedStance = FinalStance;
+		}
+		PendingBatchFinalStance.Reset();
+	}
+}
+
 void UChaosCharacterMoverComponent::ProcessSimulationEvent(const FMoverSimulationEventData& EventData)
 {
 	Super::ProcessSimulationEvent(EventData);
@@ -137,7 +156,10 @@ void UChaosCharacterMoverComponent::ProcessSimulationEvent(const FMoverSimulatio
 	}
 	if (const FStanceModifiedEventData* StanceModifiedEvent = EventData.CastTo<FStanceModifiedEventData>())
 	{
-		OnStanceModified(*StanceModifiedEvent);
+		// Accumulate - dispatch only the net final stance at end of batch (see DispatchSimulationEvents).
+		// This prevents rollback/resim transient transitions (e.g. forced uncrouch then recrouch) from
+		// causing spurious visual/animation responses on the game thread.
+		PendingBatchFinalStance = StanceModifiedEvent->NewStance;
 	}
 }
 
@@ -255,15 +277,10 @@ void UChaosCharacterMoverComponent::ProduceInputImpl(const int32 DeltaTimeMS, FM
 		ClearQueuedMode();
 	}
 
-	if (bCrouchInputPending)
+	// Always propagate the current held state so resim frames can reconstruct intent.
 	{
-		if ((bWantsToCrouch && CanCrouch()) || !bWantsToCrouch)
-		{
-			FChaosMoverCrouchInputs& CrouchInputs = Cmd->InputCollection.FindOrAddMutableDataByType<FChaosMoverCrouchInputs>();
-			CrouchInputs.bWantsToCrouch = bWantsToCrouch;
-		}
-
-		bCrouchInputPending = false;
+		FChaosMoverCrouchInputs& CrouchInputs = Cmd->InputCollection.FindOrAddMutableDataByType<FChaosMoverCrouchInputs>();
+		CrouchInputs.bWantsToCrouch = bWantsToCrouch;
 	}
 
 	if (bCancelMovementOverrides)
@@ -291,32 +308,6 @@ void UChaosCharacterMoverComponent::Launch(const FVector& VelocityOrImpulse, ECh
 	LaunchMode = Mode;
 }
 
-void UChaosCharacterMoverComponent::Crouch()
-{
-	Super::Crouch();
-	bCrouchInputPending = true;
-}
-
-void UChaosCharacterMoverComponent::UnCrouch()
-{
-	Super::UnCrouch();
-	bCrouchInputPending = true;
-}
-
-void UChaosCharacterMoverComponent::OnStanceModified(const FStanceModifiedEventData& EventData)
-{
-	// Update bWantsToCrouch without triggering networked crouch input
-	if (EventData.NewStance == EStanceMode::Invalid)
-	{
-		bWantsToCrouch = false;
-	}
-	else if (EventData.NewStance == EStanceMode::Crouch)
-	{
-		bWantsToCrouch = true;
-	}
-	
-	OnStanceChanged.Broadcast(EventData.OldStance, EventData.NewStance);
-}
 
 void UChaosCharacterMoverComponent::NotifyOnLanded(const FLandedEventData& LandedEvent)
 {
diff --git a/Engine/Plugins/Experimental/ChaosMover/Source/ChaosMover/Private/Character/Modifiers/ChaosStanceModifier.cpp b/Engine/Plugins/Experimental/ChaosMover/Source/ChaosMover/Private/Character/Modifiers/ChaosStanceModifier.cpp
--- a/Engine/Plugins/Experimental/ChaosMover/Source/ChaosMover/Private/Character/Modifiers/ChaosStanceModifier.cpp
+++ b/Engine/Plugins/Experimental/ChaosMover/Source/ChaosMover/Private/Character/Modifiers/ChaosStanceModifier.cpp
@@ -62,68 +62,63 @@ const Chaos::FCapsule* FChaosStanceModifier::GetCapsule(Chaos::FPBDRigidParticle
 
 void FChaosStanceModifier::UpdateStance(const FMovementModifierParams_Async& AsyncParams, EStanceMode NewStance)
 {
-	if (NewStance == ActiveStance)
-	{
-		// When resimulating, NewStance == ActiveStance is inherently always true because the server auth sync state will have the correct active state
-		// We will need to continue with Updating the capsule irrespective of this check
-		if (!AsyncParams.TimeStep->bIsResimulating)
-		{
-			UE_LOG(LogChaosMover, Verbose, TEXT("Skipping stance change - not required"));
-			return;
-		}
-	}
-
 	if (!AsyncParams.IsValid())
 	{
 		return;
 	}
 
-	UChaosMoverSimulation* Sim = Cast<UChaosMoverSimulation>(AsyncParams.Simulation);
-	if (!Sim)
+	UChaosMoverSimulation* Simulation = Cast<UChaosMoverSimulation>(AsyncParams.Simulation);
+	if (!Simulation)
 	{
 		return;
 	}
 
 	FName ModeName = CurrentModeName.IsNone() ? AsyncParams.SyncState->MovementMode : CurrentModeName;
-	IChaosCharacterMovementModeInterface* CharacterMode = Cast<IChaosCharacterMovementModeInterface>(Sim->FindMovementModeByName_Mutable(ModeName));
+	IChaosCharacterMovementModeInterface* CharacterMode = Cast<IChaosCharacterMovementModeInterface>(Simulation->FindMovementModeByName_Mutable(ModeName));
 	if (!CharacterMode)
 	{
 		UE_LOG(LogChaosMover, Warning, TEXT("ChaosStanceModifier only works with character modes"));
 		return;
 	}
 
-	const FChaosMoverSimulationDefaultInputs* SimInputs = Sim->GetLocalSimInput().FindDataByType<FChaosMoverSimulationDefaultInputs>();
+	const FChaosMoverSimulationDefaultInputs* SimInputs = Simulation->GetLocalSimInput().FindDataByType<FChaosMoverSimulationDefaultInputs>();
 	if (!SimInputs)
 	{
 		UE_LOG(LogChaosMover, Warning, TEXT("ChaosStanceModifier requires FCharacterDefaultInputs"));
 		return;
 	}
 
-	Chaos::FPhysicsSolver* Solver = nullptr;
-	if (SimInputs->World)
-	{
-		if (FPhysScene* Scene = SimInputs->World->GetPhysicsScene())
-		{
-			Solver = Scene->GetSolver();
-		}
-	}
-
+	FPhysScene* Scene = SimInputs->World ? SimInputs->World->GetPhysicsScene() : nullptr;
+	Chaos::FPhysicsSolver* Solver = Solver = Scene ? Scene->GetSolver() : nullptr;
 	if (!Solver)
 	{
 		UE_LOG(LogChaosMover, Warning, TEXT("ChaosStanceModifier requires FPhysicsSolver"));
 		return;
 	}
 
+	if (NewStance == ActiveStance)
+	{
+		// NewStance == ActiveStance can happen on the first resim frame after a rollback.
+		// This happens if we rollback to a frame where the modifier was active and already
+		// active on previous frames. In that case, ActiveStance was deserialized from the restored
+		// sync state and already reflects the rollback frame's stance before OnStart_Async runs. 
+		// We must still re-apply the capsule for physics correctness, so we fall through on that frame only.
+		// On any other frame (non-resim, or subsequent resim frames) this is a genuine no-op; skip it.
+		if (!AsyncParams.TimeStep->bIsResimulating || !AsyncParams.TimeStep->bIsFirstResimFrame)
+		{
+			UE_LOG(LogChaosMover, Verbose, TEXT("Skipping stance change - not required"));
+			return;
+		}
+	}
+
 	Chaos::FWritePhysicsObjectInterface_Internal WriteInterface = Chaos::FPhysicsObjectInternalInterface::GetWrite();
 	Chaos::FPBDRigidParticleHandle* ParticleHandle = WriteInterface.GetRigidParticle(SimInputs->PhysicsObject);
-
 	if (!ParticleHandle)
 	{
 		return;
 	}
 
 	const Chaos::FCapsule* Capsule = GetCapsule(ParticleHandle);
-
 	if (!Capsule)
 	{
 		UE_LOG(LogChaosMover, Warning, TEXT("No capsule shape found in ChaosStanceModifier"));
@@ -133,7 +128,6 @@ void FChaosStanceModifier::UpdateStance(const FMovementModifierParams_Async& Asy
 	// Invalidate the particle as we are about to change the collision geometry so need to clear contacts etc.
 	Solver->GetEvolution()->InvalidateParticle(ParticleHandle);
 
-
 	// Scale the capsule
 	float NewHalfHeight, NewRadius, NewGroundClearance;
 	if (NewStance == EStanceMode::Crouch)
@@ -181,7 +175,7 @@ void FChaosStanceModifier::UpdateStance(const FMovementModifierParams_Async& Asy
 	};
 
 	// Fire simulation event
-	Sim->AddEvent(MakeShared<FStanceModifiedEventData>(*AsyncParams.TimeStep, OldStance, ActiveStance, SetGeometryGTCallback));
+	Simulation->AddEvent(MakeShared<FStanceModifiedEventData>(*AsyncParams.TimeStep, OldStance, ActiveStance, SetGeometryGTCallback));
 }
 
 void FChaosStanceModifier::OnStart_Async(const FMovementModifierParams_Async& AsyncParams)
diff --git a/Engine/Plugins/Experimental/ChaosMover/Source/ChaosMover/Public/ChaosMover/Character/ChaosCharacterMoverComponent.h b/Engine/Plugins/Experimental/ChaosMover/Source/ChaosMover/Public/ChaosMover/Character/ChaosCharacterMoverComponent.h
--- a/Engine/Plugins/Experimental/ChaosMover/Source/ChaosMover/Public/ChaosMover/Character/ChaosCharacterMoverComponent.h
+++ b/Engine/Plugins/Experimental/ChaosMover/Source/ChaosMover/Public/ChaosMover/Character/ChaosCharacterMoverComponent.h
@@ -36,9 +36,6 @@ public:
 	CHAOSMOVER_API virtual void DoQueueNextMode(FName DesiredModeName, bool bShouldReenter=false) override;
 	CHAOSMOVER_API virtual FName GetNextMovementModeName() const override;
 
-	CHAOSMOVER_API virtual void Crouch() override;
-	CHAOSMOVER_API virtual void UnCrouch() override;
-
 	// Broadcast when this actor lands on a valid surface.
 	UPROPERTY(BlueprintAssignable, Category = "Chaos Mover", DisplayName = "On Landed")
 	FChaosMover_OnLanded OnLandedDelegate;
@@ -76,11 +73,11 @@ public:
 
 protected:
 	CHAOSMOVER_API virtual void OnMoverPreSimulationTick(const FMoverTimeStep& TimeStep, const FMoverInputCmdContext& InputCmd) override;
+	CHAOSMOVER_API virtual void DispatchSimulationEvents(const UE::Mover::FSimulationOutputData& OutputData) override;
 	CHAOSMOVER_API virtual void ProcessSimulationEvent(const FMoverSimulationEventData& EventData) override;
 	CHAOSMOVER_API virtual void SetAdditionalSimulationOutput(const FMoverDataCollection& Data) override;
 
 	CHAOSMOVER_API void ClearQueuedMode();
-	CHAOSMOVER_API virtual void OnStanceModified(const FStanceModifiedEventData& EventData);
 
 	CHAOSMOVER_API void NotifyOnLanded(const FLandedEventData& LandedEvent);
 	CHAOSMOVER_API virtual void OnLanded(const FLandedEventData& LandedEvent);
@@ -112,5 +109,10 @@ protected:
 	// Queued immediate mode transition, it will be transmitted to the PT next time ProduceInput is called
 	FName QueuedModeTransitionName = NAME_None;
 
-	bool bCrouchInputPending = false;
+	// Batch accumulation for stance dispatch: avoids spurious visual transitions caused by
+	// rollback/resim transient transitions (e.g. forced uncrouch then recrouch).
+	// ProcessSimulationEvent accumulates the final stance seen in a batch,
+	// and SetAdditionalSimulationOutput fires OnStanceChanged only if the net result changed.
+	EStanceMode LastDispatchedStance = EStanceMode::Invalid;
+	TOptional<EStanceMode> PendingBatchFinalStance;
 };
