ChaosMover: PT-side gameplay tag diffing and rollback correction

Depends on CL 51732208.

FMoverStateMachine: DiffAndEmitGameplayTagEvents diffs tag state each tick and emits
FMoverGameplayTagChangeEventData events; LastKnownGameplayTags tracks emitted state;
OnSimulationRollback re-diffs against rolled-back sync state and emits corrective events;
GetLastKnownGameplayTags() accessor.

ChaosMoverBackend: OnSimulationRollback calls PruneEventsFromFrame to discard stale
predicted events; calls Backend->OnSimulationRollback via ChaosMoverSubsystem on
bIsFirstResimFrame detection; DebugOwnerName initialized with PIE-aware world package suffix.

ChaosMoverSimulation: InternalGameplayTags from GetLastKnownGameplayTags() in simulation
output; bReEmitOnResim skips ProcessedEvents dedup so tag events always re-emit on resim.

#Synced-CL 52152484

diff --git a/Engine/Plugins/Experimental/ChaosMover/Source/ChaosMover/Private/Backends/ChaosMoverAsyncCallback.cpp b/Engine/Plugins/Experimental/ChaosMover/Source/ChaosMover/Private/Backends/ChaosMoverAsyncCallback.cpp
--- a/Engine/Plugins/Experimental/ChaosMover/Source/ChaosMover/Private/Backends/ChaosMoverAsyncCallback.cpp
+++ b/Engine/Plugins/Experimental/ChaosMover/Source/ChaosMover/Private/Backends/ChaosMoverAsyncCallback.cpp
@@ -90,7 +90,8 @@ namespace UE::ChaosMover
 		
 		const int32 NumInputs = AsyncInput->InputData.Num();
 		AsyncOutput.OutputData.SetNum(NumInputs);
-		AsyncOutput.TimeStep.SetNumUninitialized(NumInputs);		
+		AsyncOutput.TimeStep.SetNum(NumInputs);
+		AsyncOutput.RollbackNewSyncState.SetNum(NumInputs);
 
 		const int32 NumBackends = AsyncInput->Backends.Num();
 		ensure(NumInputs == NumBackends);
@@ -109,6 +110,11 @@ namespace UE::ChaosMover
 				SimContext.Blackboard = SimContext.Simulation->GetRollbackBlackboardSimWrapper();
 				SimContext.SimulationOwner = nullptr;	// Not required for ChaosMover modes
 
+				if (TimeStep.bIsFirstResimFrame)
+				{
+					AsyncOutput.RollbackNewSyncState[BackendIndex] = Backend->GetSimulation()->GetCurrentSyncState();
+				}
+
 				Backend->GetSimulation()->SimulationTick(SimContext, TimeStep, AsyncInput->InputData[BackendIndex], AsyncOutput.OutputData[BackendIndex]);
 			}
 		};
diff --git a/Engine/Plugins/Experimental/ChaosMover/Source/ChaosMover/Private/Backends/ChaosMoverAsyncCallback.h b/Engine/Plugins/Experimental/ChaosMover/Source/ChaosMover/Private/Backends/ChaosMoverAsyncCallback.h
--- a/Engine/Plugins/Experimental/ChaosMover/Source/ChaosMover/Private/Backends/ChaosMoverAsyncCallback.h
+++ b/Engine/Plugins/Experimental/ChaosMover/Source/ChaosMover/Private/Backends/ChaosMoverAsyncCallback.h
@@ -31,11 +31,14 @@ namespace UE::ChaosMover
 		TArray<FSimulationOutputData> OutputData;
 		TArray<TWeakObjectPtr<UChaosMoverBackendComponent>> Backends;
 		TArray<FMoverTimeStep> TimeStep;
+		// Sync state at the rollback destination frame (NewBaseTimeStep), populated only when bIsFirstResimFrame is true.
+		TArray<FMoverSyncState> RollbackNewSyncState;
 		void Reset()
 		{
 			OutputData.Empty();
 			Backends.Empty();
 			TimeStep.Empty();
+			RollbackNewSyncState.Empty();
 		}
 	};
 
diff --git a/Engine/Plugins/Experimental/ChaosMover/Source/ChaosMover/Private/Backends/ChaosMoverBackend.cpp b/Engine/Plugins/Experimental/ChaosMover/Source/ChaosMover/Private/Backends/ChaosMoverBackend.cpp
--- a/Engine/Plugins/Experimental/ChaosMover/Source/ChaosMover/Private/Backends/ChaosMoverBackend.cpp
+++ b/Engine/Plugins/Experimental/ChaosMover/Source/ChaosMover/Private/Backends/ChaosMoverBackend.cpp
@@ -3,6 +3,7 @@
 #include "ChaosMover/Backends/ChaosMoverBackend.h"
 
 #include "Backends/ChaosMoverSubsystem.h"
+#include "MoverGameplayTagLog.h"
 #include "ChaosMover/ChaosMoverConsoleVariables.h"
 #include "ChaosMover/ChaosMoverLog.h"
 #include "ChaosMover/ChaosMoverSimulation.h"
@@ -335,6 +336,18 @@ void UChaosMoverBackendComponent::InitSimulation()
 	
 	SimOutputRecord.Clear();
 
+#if !NO_LOGGING
+	// Mover Gameplay Tag Logging
+	{
+		const AActor* DebugOwner = MoverComp.GetOwner();
+		const UWorld* DebugWorld = DebugOwner ? DebugOwner->GetWorld() : nullptr;
+		SimOutputRecord.DebugOwnerName = FString::Printf(TEXT("%s[%s]"), *GetNameSafe(DebugOwner), DebugWorld ? *DebugWorld->GetOutermost()->GetName() : TEXT("null"));
+	}
+	SimOutputRecord.DebugOwnerRole = MoverComp.GetOwner() ? StaticEnum<ENetRole>()->GetNameStringByValue((int64)MoverComp.GetOwner()->GetLocalRole()) : TEXT("null");
+	Simulation->DebugOwnerName = SimOutputRecord.DebugOwnerName;
+	Simulation->DebugOwnerRole = SimOutputRecord.DebugOwnerRole;
+#endif
+
 	UE::ChaosMover::FSimulationOutputData OutputData;
 	FMoverAuxStateContext UnusedAuxState;
 	MoverComp.InitializeSimulationState(&OutputData.SyncState, &UnusedAuxState);
@@ -762,9 +775,33 @@ void UChaosMoverBackendComponent::ConsumeOutputData(const FMoverTimeStep& TimeSt
 	if (OutputData.IsValid())
 	{
 		SimOutputRecord.Add(TimeStep, OutputData);
+
+		if (!TimeStep.bIsResimulating)
+		{
+			PreRollbackTimeStep = TimeStep;
+		}
+
+		// Mover Gameplay Tag Logging
+		MOVER_TAG_LOG(LogChaosMover, "[GT:Consume] Actor=%s Role=%s Frame=%d TimeMs=%.3f IsResim=%d IsFirstResim=%d | EventsInRecord=%d",
+			*Simulation->DebugOwnerName, *Simulation->DebugOwnerRole, TimeStep.ServerFrame, TimeStep.BaseSimTimeMs, (int32)TimeStep.bIsResimulating, (int32)TimeStep.bIsFirstResimFrame, 0);
 	}
 }
 
+void UChaosMoverBackendComponent::OnSimulationRollback(const FMoverSyncState& NewSyncState, const FMoverTimeStep& NewBaseTimeStep, const FMoverTimeStep& InPreRollbackTimeStep)
+{
+	Chaos::EnsureIsInGameThreadContext();
+
+	// Discard predicted events from the frames being re-simulated.
+	// Gameplay Tag Specific: the corrective rollback events produced by DiffAndEmitGameplayTagEvents
+	// arrive in the first resim frame's ConsumeOutputData call and will be added then, so they are
+	// unaffected by this pruning.
+	SimOutputRecord.PruneEventsFromFrame(NewBaseTimeStep.ServerFrame);
+
+	// Mover Gameplay Tag Logging
+	MOVER_TAG_LOG(LogChaosMover, "[GT:SimRecordPrune] Actor=%s Role=%s Frame=%d IsResim=%d IsRollback=%d | pruned events with ServerFrame>=%d",
+		*Simulation->DebugOwnerName, *Simulation->DebugOwnerRole, NewBaseTimeStep.ServerFrame, (int32)NewBaseTimeStep.bIsResimulating, (int32)true, NewBaseTimeStep.ServerFrame);
+}
+
 void UChaosMoverBackendComponent::FinalizeFrame(double ResultsTimeInMs)
 {
 	Chaos::EnsureIsInGameThreadContext();
@@ -786,6 +823,10 @@ void UChaosMoverBackendComponent::FinalizeFrame(double ResultsTimeInMs)
 	UE::ChaosMover::FSimulationOutputData InterpolatedOutput;
 	// ResultsTimeInMs is the time at the end of the step while CreateInterpolatedResult expects the time at the beginning of the step so we correct by subtracting one async dt
 	const double AtBaseTimeMs = ResultsTimeInMs - Solver->GetAsyncDeltaTime() * 1000.0;
+
+	// Mover Gameplay Tag Logging
+	MOVER_TAG_LOG(LogChaosMover, "[GT:Finalize] Actor=%s Role=%s AtTimeMs=%.3f", *SimOutputRecord.DebugOwnerName, *Simulation->DebugOwnerRole, AtBaseTimeMs);
+
 	SimOutputRecord.CreateInterpolatedResult(AtBaseTimeMs, /*OUT*/ TimeStep, /*OUT*/ InterpolatedOutput);
 
 	// Physics interactions in the last frame may have caused a change in position or velocity that's different from what a simple lerp would predict,
diff --git a/Engine/Plugins/Experimental/ChaosMover/Source/ChaosMover/Private/Backends/ChaosMoverSubsystem.cpp b/Engine/Plugins/Experimental/ChaosMover/Source/ChaosMover/Private/Backends/ChaosMoverSubsystem.cpp
--- a/Engine/Plugins/Experimental/ChaosMover/Source/ChaosMover/Private/Backends/ChaosMoverSubsystem.cpp
+++ b/Engine/Plugins/Experimental/ChaosMover/Source/ChaosMover/Private/Backends/ChaosMoverSubsystem.cpp
@@ -133,6 +133,14 @@ void UChaosMoverSubsystem::OnPostPhysicsTick(FChaosScene* Scene)
 		auto LambdaParallelUpdate = [&AsyncOutput](int32 Idx) {
 			if (UChaosMoverBackendComponent* Backend = AsyncOutput->Backends[Idx].Get())
 			{
+				// In the async path, UMoverComponent::OnSimulationRollback is never called on the GT.
+				// The subsystem is the first to know about a rollback via bIsFirstResimFrame, so
+				// notify the backend here before consuming output for this resim sequence.
+				if (AsyncOutput->TimeStep[Idx].bIsFirstResimFrame)
+				{
+					Backend->OnSimulationRollback(AsyncOutput->RollbackNewSyncState[Idx], AsyncOutput->TimeStep[Idx], Backend->GetPreRollbackTimeStep());
+				}
+
 				Backend->ConsumeOutputData(AsyncOutput->TimeStep[Idx], AsyncOutput->OutputData[Idx]);
 			}
 		};
diff --git a/Engine/Plugins/Experimental/ChaosMover/Source/ChaosMover/Private/ChaosMoverConsoleVariables.cpp b/Engine/Plugins/Experimental/ChaosMover/Source/ChaosMover/Private/ChaosMoverConsoleVariables.cpp
--- a/Engine/Plugins/Experimental/ChaosMover/Source/ChaosMover/Private/ChaosMoverConsoleVariables.cpp
+++ b/Engine/Plugins/Experimental/ChaosMover/Source/ChaosMover/Private/ChaosMoverConsoleVariables.cpp
@@ -21,7 +21,7 @@ namespace UE::ChaosMover
 
 		bool bForceSingleThreadedPT = true;
 		FAutoConsoleVariableRef CVarChaosMoverForceSingleThreadedPT(TEXT("ChaosMover.ForceSingleThreadedPT"),
-			bForceSingleThreadedPT, TEXT("Force updates on the physics thread to be single threaded."));
+			bForceSingleThreadedPT, TEXT("Force updates on the worker thread to be single threaded."));
 
 		bool bDrawGroundQueries = false;
 		FAutoConsoleVariableRef CVarChaosMoverDrawGroundQueries(TEXT("ChaosMover.DebugDraw.GroundQueries"),
diff --git a/Engine/Plugins/Experimental/ChaosMover/Source/ChaosMover/Private/ChaosMoverSimulation.cpp b/Engine/Plugins/Experimental/ChaosMover/Source/ChaosMover/Private/ChaosMoverSimulation.cpp
--- a/Engine/Plugins/Experimental/ChaosMover/Source/ChaosMover/Private/ChaosMoverSimulation.cpp
+++ b/Engine/Plugins/Experimental/ChaosMover/Source/ChaosMover/Private/ChaosMoverSimulation.cpp
@@ -37,6 +37,7 @@
 #include "PhysicsProxy/CharacterGroundConstraintProxy.h"
 #include "UObject/UObjectGlobals.h"
 #include "Chaos/GeometryParticles.h"
+#include "MoverGameplayTagLog.h"
 
 #if WITH_CHAOS_VISUAL_DEBUGGER
 #include "ChaosVisualDebugger/ChaosVisualDebuggerTrace.h"
@@ -611,6 +612,10 @@ void UChaosMoverSimulation::OnSimulationRollback(const FMoverTimeStep& NewTimeSt
 {
 	QUICK_SCOPE_CYCLE_COUNTER(ChaosMoverSimulation_OnSimulationRollback);
 
+	// Mover Gameplay Tag Logging
+	MOVER_TAG_LOG(LogChaosMover, "[WT:RollbackTrigger] Actor=%s Role=%s Frame=%d IsResim=%d IsRollback=%d | rolling back to this frame",
+		*DebugOwnerName, *DebugOwnerRole, NewTimeStep.ServerFrame, (int32)NewTimeStep.bIsResimulating, (int32)true);
+
 	// Rollback blackboard on the first frame of resimulation
 	Blackboard->Invalidate(EInvalidationReason::Rollback);
 	RollbackBlackboard->BeginRollback(NewTimeStep);
@@ -618,6 +623,11 @@ void UChaosMoverSimulation::OnSimulationRollback(const FMoverTimeStep& NewTimeSt
 	// Remove any events that have been stored for this simulation tick
 	Events.Empty();
 
+	// Mover Gameplay Tag Logging
+	MOVER_TAG_LOG(LogChaosMover, "[WT:EventBufClear] Actor=%s Role=%s Frame=%d IsResim=%d IsRollback=%d | in-flight Events cleared",
+		*DebugOwnerName, *DebugOwnerRole, NewTimeStep.ServerFrame, (int32)NewTimeStep.bIsResimulating, (int32)true);
+
+
 	// Make sure the movement basis matches the one in the sync state
 	if (FChaosMovementBasis* MovementBasis = CurrentSyncState.SyncStateCollection.FindMutableDataByType<FChaosMovementBasis>())
 	{
@@ -676,6 +686,7 @@ void UChaosMoverSimulation::OnSimulationTick(const FMoverSimContext& SimContext,
 	// Copy the sync state locally and to the output data
 	OutputData.SyncState = TickEndData.SyncState;
 	OutputData.LastUsedInputCmd = InputData.InputCmd;
+	OutputData.InternalGameplayTags = StateMachine.GetLastKnownGameplayTags();
 }
 
 void UChaosMoverSimulation::PostSimulationTickCharacter(const IChaosCharacterMovementModeInterface& CharacterMode, const FMoverTimeStep& TimeStep, UE::ChaosMover::FSimulationOutputData& OutputData)
@@ -1124,11 +1135,11 @@ void UChaosMoverSimulation::UpdateEvents(const FMoverTimeStep& TimeStep, UE::Cha
 
 				if (TimeStep.bIsResimulating)
 				{
-					// Remove events that have already been processed in the main update
+					// Suppress re-emission of events the GT already received during the forward sim.
+					// Events with bReEmitOnResim=true are never in ProcessedEvents and always pass through.
 
 					Events.RemoveAll([this](const TSharedPtr<FMoverSimulationEventData> Event) {
-						// Stance events are handled by batch accumulation on the game thread - skip dedup
-						if (Event->CastTo<FStanceModifiedEventData>())
+						if (Event->bReEmitOnResim)
 						{
 							return false;
 						}
@@ -1149,10 +1160,14 @@ void UChaosMoverSimulation::UpdateEvents(const FMoverTimeStep& TimeStep, UE::Cha
 				}
 				else
 				{
-					// Add new events to the processed list
+					// Record only events that should not re-emit on resim. Events with bReEmitOnResim=true
+					// are intentionally excluded so they always pass the dedup check and re-emit on resim.
 					for (TSharedPtr<FMoverSimulationEventData>& Event : Events)
 					{
-						ProcessedEvents.Add(Event);
+						if (!Event->bReEmitOnResim)
+						{
+							ProcessedEvents.Add(Event);
+						}
 					}
 				}
 			}
@@ -1168,7 +1183,7 @@ void UChaosMoverSimulation::TraceMoverData(const FMoverTimeStep& TimeStep, const
 {
 	QUICK_SCOPE_CYCLE_COUNTER(ChaosMoverSimulation_TraceMoverData);
 
-	// Send the latest physics thread data to CVD
+	// Send the latest worker thread data to CVD
 #if WITH_CHAOS_VISUAL_DEBUGGER
 	if (FChaosVisualDebuggerTrace::IsTracing())
 	{
@@ -1441,7 +1456,7 @@ DEFINE_FUNCTION(UChaosMoverSimulation::execK2_QueueInstantMovementEffect)
 		FInstantMovementEffect* EffectAsBasePtr = reinterpret_cast<FInstantMovementEffect*>(EffectPtr);
 		FInstantMovementEffect* ClonedMove = EffectAsBasePtr->Clone();
 
-		P_THIS->QueueInstantMovementEffect(TSharedPtr<FInstantMovementEffect>(ClonedMove), /*bShouldRollBack =*/ true); // Physics Thread emitted effects can be rolled back because we expect resim to rerun all Physics Thread logic
+		P_THIS->QueueInstantMovementEffect(TSharedPtr<FInstantMovementEffect>(ClonedMove), /*bShouldRollBack =*/ true); // Worker Thread emitted effects can be rolled back because we expect resim to rerun all Worker Thread logic
 	}
 
 	P_NATIVE_END;
@@ -1490,7 +1505,7 @@ DEFINE_FUNCTION(UChaosMoverSimulation::execK2_QueueLayeredMove)
 		FLayeredMoveBase* MoveAsBasePtr = reinterpret_cast<FLayeredMoveBase*>(MovePtr);
 		FLayeredMoveBase* ClonedMove = MoveAsBasePtr->Clone();
 
-		P_THIS->QueueLayeredMove(TSharedPtr<FLayeredMoveBase>(ClonedMove), /*bShouldRollBack =*/ true); // Physics Thread moves can be rolled back because we expect resim to rerun all Physics Thread logic
+		P_THIS->QueueLayeredMove(TSharedPtr<FLayeredMoveBase>(ClonedMove), /*bShouldRollBack =*/ true); // Worker Thread moves can be rolled back because we expect resim to rerun all Worker Thread logic
 	}
 
 	P_NATIVE_END;
diff --git a/Engine/Plugins/Experimental/ChaosMover/Source/ChaosMover/Private/ChaosMoverStateMachine.cpp b/Engine/Plugins/Experimental/ChaosMover/Source/ChaosMover/Private/ChaosMoverStateMachine.cpp
--- a/Engine/Plugins/Experimental/ChaosMover/Source/ChaosMover/Private/ChaosMoverStateMachine.cpp
+++ b/Engine/Plugins/Experimental/ChaosMover/Source/ChaosMover/Private/ChaosMoverStateMachine.cpp
@@ -3,6 +3,7 @@
 
 #include "ChaosMover/ChaosMoverStateMachine.h"
 
+#include "MoverGameplayTagLog.h"
 #include "Backends/ChaosMoverSubsystem.h"
 #include "ChaosMover/ChaosMoverConsoleVariables.h"
 #include "ChaosMover/ChaosMoverLog.h"
@@ -13,6 +14,7 @@
 #include "MovementModeTransition.h"
 #include "MoverComponent.h"
 #include "MoverDeveloperSettings.h"
+#include "MoverSimulationTypes.h"
 #include "MoverLog.h"
 #include "Templates/SubclassOf.h"
 #include "Framework/Threading.h"
@@ -386,6 +388,71 @@ namespace UE::ChaosMover
 		}
 	}
 
+	void FMoverStateMachine::CollectGameplayTags(const FMoverSyncState& SyncState, FGameplayTagContainer& OutTags) const
+	{
+		if (const TWeakObjectPtr<UBaseMovementMode>* ModeWeakPtr = Modes.Find(SyncState.MovementMode))
+		{
+			if (TStrongObjectPtr<UBaseMovementMode> Mode = ModeWeakPtr->Pin())
+			{
+				Mode->GetGameplayTags(OutTags);
+			}
+		}
+
+		for (auto ModifierIt = SyncState.MovementModifiers.GetActiveModifiersIterator(); ModifierIt; ++ModifierIt)
+		{
+			(*ModifierIt)->GetGameplayTags(OutTags);
+		}
+
+		for (const TSharedPtr<FLayeredMoveBase>& Move : SyncState.LayeredMoves.GetActiveMoves())
+		{
+			Move->GetGameplayTags(OutTags);
+		}
+	}
+
+	void FMoverStateMachine::DiffAndEmitGameplayTagEvents(const FGameplayTagContainer& OldTags, const FGameplayTagContainer& NewTags, const FMoverTimeStep& TimeStep, bool bIsCausedByRollback)
+	{
+		TStrongObjectPtr<UChaosMoverSimulation> SimStrongObjPtr = Simulation.Pin();
+		if (!SimStrongObjPtr)
+		{
+			return;
+		}
+
+		for (const FGameplayTag& Tag : NewTags)
+		{
+			if (!OldTags.HasTagExact(Tag))
+			{
+				TSharedPtr<FMoverGameplayTagChangeEventData> Event = MakeShared<FMoverGameplayTagChangeEventData>(TimeStep, Tag, /*bWasAdded=*/true);
+				Event->Context.bIsCausedByRollback = bIsCausedByRollback;
+
+				// Mover Gameplay Tag Logging
+				MOVER_TAG_LOG(LogChaosMover, "[WT:TagDiff] Actor=%s Role=%s Frame=%d IsResim=%d IsRollback=%d Tag=%s Change=Added",
+					*SimStrongObjPtr->DebugOwnerName, *SimStrongObjPtr->DebugOwnerRole, TimeStep.ServerFrame, (int32)TimeStep.bIsResimulating, (int32)bIsCausedByRollback, *Tag.ToString());
+
+				SimStrongObjPtr->AddEvent(Event);
+			}
+		}
+
+		for (const FGameplayTag& Tag : OldTags)
+		{
+			if (!NewTags.HasTagExact(Tag))
+			{
+				TSharedPtr<FMoverGameplayTagChangeEventData> Event = MakeShared<FMoverGameplayTagChangeEventData>(TimeStep, Tag, /*bWasAdded=*/false);
+				Event->Context.bIsCausedByRollback = bIsCausedByRollback;
+
+				// Mover Gameplay Tag Logging
+				MOVER_TAG_LOG(LogChaosMover, "[WT:TagDiff] Actor=%s Role=%s Frame=%d IsResim=%d IsRollback=%d Tag=%s Change=Removed",
+					*SimStrongObjPtr->DebugOwnerName, *SimStrongObjPtr->DebugOwnerRole, TimeStep.ServerFrame, (int32)TimeStep.bIsResimulating, (int32)bIsCausedByRollback, *Tag.ToString());
+
+				SimStrongObjPtr->AddEvent(Event);
+			}
+		}
+
+		// Mover Gameplay Tag Logging
+		MOVER_TAG_LOG(LogChaosMover, "[WT:TagDiffState] Actor=%s Role=%s Frame=%d IsResim=%d IsRollback=%d OldTags=[%s] NewTags=[%s]",
+			*SimStrongObjPtr->DebugOwnerName, *SimStrongObjPtr->DebugOwnerRole, TimeStep.ServerFrame, (int32)TimeStep.bIsResimulating, (int32)bIsCausedByRollback,
+			*OldTags.ToStringSimple(), *NewTags.ToStringSimple());
+	}
+
 	void FMoverStateMachine::OnSimulationTick(const FMoverSimContext& SimContext, const FMoverTimeStep& TimeStep, const FMoverTickStartData& StartState, UMoverBlackboard* SimBlackboard, UMovementMixer* MovementMixer, FMoverTickEndData& OutputState)
 	{
 		QUICK_SCOPE_CYCLE_COUNTER(MoverStateMachine_OnSimulationTick);
@@ -690,6 +757,14 @@ namespace UE::ChaosMover
 				Modifier->OnPostMovement_Async(ModifierParams);
 			}
 
+			// Collect all tags from active mode, modifiers, and layered moves, then emit events for any changes
+			{
+				FGameplayTagContainer NewSimTags;
+				CollectGameplayTags(OutputState.SyncState, NewSimTags);
+				DiffAndEmitGameplayTagEvents(LastKnownGameplayTags, NewSimTags, SubTimeStep, /*bIsCausedByRollback*/ false);
+				LastKnownGameplayTags = MoveTemp(NewSimTags);
+			}
+
 			const float RemainingMs = FMath::Clamp(OutputState.MovementEndState.RemainingMs, 0.0f, SubTimeStep.StepMs);
 			const float SubstepUsedMs = (SubTimeStep.StepMs - RemainingMs);
 			SubTimeStep.BaseSimTimeMs += SubstepUsedMs;
@@ -738,6 +813,22 @@ namespace UE::ChaosMover
 #endif // !UE_BUILD_SHIPPING
 	}
 
+	static void LogPTRollbackTagDiff(
+		const TWeakObjectPtr<UChaosMoverSimulation>& Simulation,
+		const FMoverTimeStep& NewTimeStep,
+		const FGameplayTagContainer& OldTags,
+		const FGameplayTagContainer& NewSimTags)
+	{
+		if (MoverCVars::EnableGameplayTagLog)
+		{
+			if (TStrongObjectPtr<UChaosMoverSimulation> DbgSim = Simulation.Pin())
+			{
+				MOVER_TAG_LOG(LogChaosMover, "[WT:RollbackTagDiff] Actor=%s Role=%s Frame=%d IsResim=%d IsRollback=%d | OldTags=[%s] AuthTags=[%s]",
+					*DbgSim->DebugOwnerName, *DbgSim->DebugOwnerRole, NewTimeStep.ServerFrame, (int32)NewTimeStep.bIsResimulating, (int32)true, *OldTags.ToStringSimple(), *NewSimTags.ToStringSimple());
+			}
+		}
+	}
+
 	void FMoverStateMachine::OnSimulationRollback(const FMoverTimeStep& NewTimeStep, const FMoverSyncState& InvalidSyncState, const FMoverSyncState& NewSyncState)
 	{
 		QUICK_SCOPE_CYCLE_COUNTER(MoverStateMachine_OnSimulationRollback);
@@ -803,6 +894,18 @@ namespace UE::ChaosMover
 				ModifierFromCacheIt->Get()->OnEnd_Async(ModifierParams);
 			}
 		}
+
+		// Collect tags from the rolled-back state and emit corrective tag events
+		{
+			FGameplayTagContainer NewSimTags;
+			CollectGameplayTags(NewSyncState, NewSimTags);
+
+			// Mover Gameplay Tag Logging
+			LogPTRollbackTagDiff(Simulation, NewTimeStep, LastKnownGameplayTags, NewSimTags);
+
+			DiffAndEmitGameplayTagEvents(LastKnownGameplayTags, NewSimTags, NewTimeStep, /*bIsCausedByRollback*/ true);
+			LastKnownGameplayTags = MoveTemp(NewSimTags);
+		}
 	}
 
 	void FMoverStateMachine::RollBackQueuedInstantMovementEffects(int32 NewServerFrame)
diff --git a/Engine/Plugins/Experimental/ChaosMover/Source/ChaosMover/Public/ChaosMover/Backends/ChaosMoverBackend.h b/Engine/Plugins/Experimental/ChaosMover/Source/ChaosMover/Public/ChaosMover/Backends/ChaosMoverBackend.h
--- a/Engine/Plugins/Experimental/ChaosMover/Source/ChaosMover/Public/ChaosMover/Backends/ChaosMoverBackend.h
+++ b/Engine/Plugins/Experimental/ChaosMover/Source/ChaosMover/Public/ChaosMover/Backends/ChaosMoverBackend.h
@@ -49,10 +49,13 @@ public:
 	CHAOSMOVER_API virtual void ProduceInputData(int32 PhysicsStep, int32 NumSteps, const FMoverTimeStep& TimeStep, UE::ChaosMover::FSimulationInputData& InputData);
 	CHAOSMOVER_API virtual void ConsumeOutputData(const FMoverTimeStep& TimeStep, const UE::ChaosMover::FSimulationOutputData& OutputData);
 	CHAOSMOVER_API virtual void FinalizeFrame(double ResultsTimeInMs);
+	CHAOSMOVER_API virtual void OnSimulationRollback(const FMoverSyncState& NewSyncState, const FMoverTimeStep& NewBaseTimeStep, const FMoverTimeStep& InPreRollbackTimeStep) override;
 
 	CHAOSMOVER_API UChaosMoverSimulation* GetSimulation();
 	CHAOSMOVER_API const UChaosMoverSimulation* GetSimulation() const;
 
+	const FMoverTimeStep& GetPreRollbackTimeStep() const { return PreRollbackTimeStep; }
+
 protected:
 	CHAOSMOVER_API virtual void InitSimulation();
 	CHAOSMOVER_API virtual void DeinitSimulation();
@@ -95,6 +98,9 @@ protected:
 
 	UE::Mover::FSimulationOutputRecord SimOutputRecord;
 
+	// Last forward-simulated (non-resim) timestep simulated before the latest rollback; passed as PreRollbackTimeStep in OnSimulationRollback.
+	FMoverTimeStep PreRollbackTimeStep;
+
 	// Character ground constraint, for moving on ground like characters
 	TUniquePtr<Chaos::FCharacterGroundConstraint> CharacterGroundConstraint;
 
diff --git a/Engine/Plugins/Experimental/ChaosMover/Source/ChaosMover/Public/ChaosMover/ChaosMoverSimulation.h b/Engine/Plugins/Experimental/ChaosMover/Source/ChaosMover/Public/ChaosMover/ChaosMoverSimulation.h
--- a/Engine/Plugins/Experimental/ChaosMover/Source/ChaosMover/Public/ChaosMover/ChaosMoverSimulation.h
+++ b/Engine/Plugins/Experimental/ChaosMover/Source/ChaosMover/Public/ChaosMover/ChaosMoverSimulation.h
@@ -36,6 +36,12 @@ class UChaosMoverSimulation : public UMoverSimulation
 	GENERATED_BODY()
 
 public:
+#if !NO_LOGGING
+	// Mover Gameplay Tag Logging
+	FString DebugOwnerName;
+	FString DebugOwnerRole;
+#endif
+
 	CHAOSMOVER_API UChaosMoverSimulation();
 
 	// Returns the local simulation input MoverDataCollection, to read local non networked data passed to the simulation by the gameplay thread
@@ -181,6 +187,8 @@ public:
 	// Trace mover data to Chaos Visual Debugger
 	CHAOSMOVER_API void TraceMoverData(const FMoverTimeStep& TimeStep, const UE::ChaosMover::FSimulationOutputData& OutputData);
 
+	const FMoverSyncState& GetCurrentSyncState() const { return CurrentSyncState; }
+
 protected:
 	// Object State.
 	// In kinematic mode, the updated primitive component needs both its SimulatePhysics and UpdateKinematicFromSimulation options checked
@@ -241,7 +249,9 @@ private:
 	// Events from this simulation update
 	TArray<TSharedPtr<FMoverSimulationEventData>> Events;
 
-	// Events that have been sent to GT but are still within the resimulation window
+	// Dedup record of events sent to the GT during the forward sim, used to suppress duplicate
+	// re-emission during resimulation. Only contains events with bReEmitOnResim=false; events
+	// that opt into re-emission are never added here and always pass through resim unchanged.
 	TArray<TSharedPtr<FMoverSimulationEventData>> ProcessedEvents;
 
 	bool bInputCmdOverridden = false;
diff --git a/Engine/Plugins/Experimental/ChaosMover/Source/ChaosMover/Public/ChaosMover/ChaosMoverSimulationTypes.h b/Engine/Plugins/Experimental/ChaosMover/Source/ChaosMover/Public/ChaosMover/ChaosMoverSimulationTypes.h
--- a/Engine/Plugins/Experimental/ChaosMover/Source/ChaosMover/Public/ChaosMover/ChaosMoverSimulationTypes.h
+++ b/Engine/Plugins/Experimental/ChaosMover/Source/ChaosMover/Public/ChaosMover/ChaosMoverSimulationTypes.h
@@ -464,8 +464,12 @@ struct FStanceModifiedEventData : public FMoverSimulationEventData
 		, OldStance(InOldStance)
 		, NewStance(InNewStance)
 	{
+		bReEmitOnResim = true;
+	}
+	FStanceModifiedEventData()
+	{
+		bReEmitOnResim = true;
 	}
-	FStanceModifiedEventData() {}
 
 	virtual UScriptStruct* GetScriptStruct() const override
 	{
diff --git a/Engine/Plugins/Experimental/ChaosMover/Source/ChaosMover/Public/ChaosMover/ChaosMoverStateMachine.h b/Engine/Plugins/Experimental/ChaosMover/Source/ChaosMover/Public/ChaosMover/ChaosMoverStateMachine.h
--- a/Engine/Plugins/Experimental/ChaosMover/Source/ChaosMover/Public/ChaosMover/ChaosMoverStateMachine.h
+++ b/Engine/Plugins/Experimental/ChaosMover/Source/ChaosMover/Public/ChaosMover/ChaosMoverStateMachine.h
@@ -51,6 +51,11 @@ namespace UE::ChaosMover
 			return CurrentModeName;
 		}
 
+		const FGameplayTagContainer& GetLastKnownGameplayTags() const
+		{
+			return LastKnownGameplayTags;
+		}
+
 		CHAOSMOVER_API const TWeakObjectPtr<UBaseMovementMode> GetCurrentMode() const;
 		CHAOSMOVER_API const TWeakObjectPtr<UBaseMovementMode> FindMovementMode(FName ModeName) const;
 		CHAOSMOVER_API TWeakObjectPtr<UBaseMovementMode> FindMovementMode_Mutable(FName ModeName);
@@ -116,8 +121,17 @@ namespace UE::ChaosMover
 
 		TWeakObjectPtr<UChaosMoverSimulation> Simulation;
 
+		/** Last set of simulation tags that were diffed and emitted. Used to detect tag changes between substeps and on rollback. */
+		FGameplayTagContainer LastKnownGameplayTags;
+
 	private:
 		void ConstructDefaultModes();
+
+		/** Collect all gameplay tags contributed by the active mode, modifiers, and layered moves in SyncState. */
+		void CollectGameplayTags(const FMoverSyncState& SyncState, FGameplayTagContainer& OutTags) const;
+
+		/** Diff OldTags vs NewTags and emit FMoverGameplayTagChangeEventData for each change (bWasAdded=true for additions, false for removals). */
+		void DiffAndEmitGameplayTagEvents(const FGameplayTagContainer& OldTags, const FGameplayTagContainer& NewTags, const FMoverTimeStep& TimeStep, bool bIsCausedByRollback);
 		void AdvanceToNextMode(const FMoverTimeStep& TimeStep);
 		void FlushQueuedMovesToGroup(const FMoverTimeStep& TimeStep, FLayeredMoveGroup& Group);
 		void FlushQueuedModifiersToGroup(FMovementModifierGroup& ModifierGroup);
diff --git a/Engine/Plugins/Experimental/ChaosMover/Source/ChaosMover/Public/ChaosMover/PathedMovement/ChaosPathedMovementMode.h b/Engine/Plugins/Experimental/ChaosMover/Source/ChaosMover/Public/ChaosMover/PathedMovement/ChaosPathedMovementMode.h
--- a/Engine/Plugins/Experimental/ChaosMover/Source/ChaosMover/Public/ChaosMover/PathedMovement/ChaosPathedMovementMode.h
+++ b/Engine/Plugins/Experimental/ChaosMover/Source/ChaosMover/Public/ChaosMover/PathedMovement/ChaosPathedMovementMode.h
@@ -83,7 +83,7 @@ private:
 	 * If false, no joint will be used, and the affected component will be moved kinematically (i.e. independent of other physics effects and treated as though it has
 	 * infinite mass in the physics sim) and never deviate from the path.
 	 *
-	 * Note that the movement is calculated on the physics thread either way, so it will correctly interact with other bodies in the sim.
+	 * Note that the movement is calculated on the worker thread either way, so it will correctly interact with other bodies in the sim.
 	 *
 	 * Ex: Platform moving along a path in a straight line. If this is enabled, the platform could wobble, dip, shake, stall, etc. in response to physics forces/impulses
 	 *		(depending on the joint and physics properties of the platform). If this is disabled, the platform will rigidly move along the path and ignore all physics.
