ChaosMover: Fix for a self assignment, probably a copy paste mistake?

#Synced-CL 52152484

diff --git a/Engine/Plugins/Experimental/ChaosMover/Source/ChaosMover/Private/Character/Modifiers/ChaosStanceModifier.cpp b/Engine/Plugins/Experimental/ChaosMover/Source/ChaosMover/Private/Character/Modifiers/ChaosStanceModifier.cpp
--- a/Engine/Plugins/Experimental/ChaosMover/Source/ChaosMover/Private/Character/Modifiers/ChaosStanceModifier.cpp
+++ b/Engine/Plugins/Experimental/ChaosMover/Source/ChaosMover/Private/Character/Modifiers/ChaosStanceModifier.cpp
@@ -89,7 +89,7 @@ void FChaosStanceModifier::UpdateStance(const FMovementModifierParams_Async& Asy
 	}
 
 	FPhysScene* Scene = SimInputs->World ? SimInputs->World->GetPhysicsScene() : nullptr;
-	Chaos::FPhysicsSolver* Solver = Solver = Scene ? Scene->GetSolver() : nullptr;
+	Chaos::FPhysicsSolver* Solver = Scene ? Scene->GetSolver() : nullptr;
 	if (!Solver)
 	{
 		UE_LOG(LogChaosMover, Warning, TEXT("ChaosStanceModifier requires FPhysicsSolver"));
