ChaosMover: Fix PathedMovementMode to work in standalone/variable dt mode.

#Synced-CL 52152484

diff --git a/Engine/Plugins/Experimental/ChaosMover/Source/ChaosMover/Private/Backends/ChaosMoverBackend.cpp b/Engine/Plugins/Experimental/ChaosMover/Source/ChaosMover/Private/Backends/ChaosMoverBackend.cpp
--- a/Engine/Plugins/Experimental/ChaosMover/Source/ChaosMover/Private/Backends/ChaosMoverBackend.cpp
+++ b/Engine/Plugins/Experimental/ChaosMover/Source/ChaosMover/Private/Backends/ChaosMoverBackend.cpp
@@ -65,6 +65,11 @@ void UChaosMoverBackendComponent::InitializeComponent()
 
 	if (UWorld* World = GetWorld(); World && World->IsGameWorld())
 	{
+		if (const Chaos::FPhysicsSolver* Solver = GetPhysicsSolver())
+		{
+			bIsUsingAsyncPhysics = Solver->IsUsingAsyncResults();
+		}
+
 		UMoverComponent& MoverComp = GetMoverComponent();
 
 		Simulation = NewObject<UChaosMoverSimulation>(GetOwner(), SimulationClass, TEXT("ChaosMoverSimulation"), RF_Transient);
@@ -88,7 +93,7 @@ void UChaosMoverBackendComponent::InitializeComponent()
 			{
 				NetworkPhysicsComponent->InitializeComponent();
 			}
-			NetworkPhysicsComponent->Activate(true);
+			NetworkPhysicsComponent->Activate(/* bReset = */ true);
 
 			// Register network data for recording and rewind/resim
 			NetworkPhysicsComponent->CreateDataHistory<UE::ChaosMover::FNetworkDataTraits>(this);
@@ -333,6 +338,7 @@ void UChaosMoverBackendComponent::InitSimulation()
 	Params.PhysicsObject = GetPhysicsObject();
 	Params.Solver = GetPhysicsSolver();
 	Params.World = GetWorld();
+	Params.bIsUsingAsyncPhysics = bIsUsingAsyncPhysics;
 	
 	SimOutputRecord.Clear();
 
@@ -477,12 +483,12 @@ FMoverTime UChaosMoverBackendComponent::GetScheduledNetworkTime(const FMoverTime
 
 bool UChaosMoverBackendComponent::IsAsync() const
 {
-	return true;
+	return bIsUsingAsyncPhysics;
 }
 
 bool UChaosMoverBackendComponent::IsFixedDt() const
 {
-	return true;
+	return bIsUsingAsyncPhysics;
 }
 
 bool UChaosMoverBackendComponent::ShouldResim() const
diff --git a/Engine/Plugins/Experimental/ChaosMover/Source/ChaosMover/Private/Backends/ChaosMoverSubsystem.cpp b/Engine/Plugins/Experimental/ChaosMover/Source/ChaosMover/Private/Backends/ChaosMoverSubsystem.cpp
--- a/Engine/Plugins/Experimental/ChaosMover/Source/ChaosMover/Private/Backends/ChaosMoverSubsystem.cpp
+++ b/Engine/Plugins/Experimental/ChaosMover/Source/ChaosMover/Private/Backends/ChaosMoverSubsystem.cpp
@@ -213,7 +213,7 @@ FMoverTimeStep UChaosMoverSubsystem::GetMoverTimeStep() const
 					TimeStep.StepMs = FMath::Clamp(World->GetDeltaSeconds(), Solver->GetMinDeltaTime_External(), Solver->GetMaxDeltaTime_External()) * 1000.0f;
 				}
 
-				TimeStep.ServerFrame += GetNetworkPhysicsTickOffset();;
+				TimeStep.ServerFrame += GetNetworkPhysicsTickOffset();
 			}
 		}
 	}
diff --git a/Engine/Plugins/Experimental/ChaosMover/Source/ChaosMover/Private/ChaosMoverSimulation.cpp b/Engine/Plugins/Experimental/ChaosMover/Source/ChaosMover/Private/ChaosMoverSimulation.cpp
--- a/Engine/Plugins/Experimental/ChaosMover/Source/ChaosMover/Private/ChaosMoverSimulation.cpp
+++ b/Engine/Plugins/Experimental/ChaosMover/Source/ChaosMover/Private/ChaosMoverSimulation.cpp
@@ -122,6 +122,7 @@ void UChaosMoverSimulation::Init(const FInitParams& InitParams)
 	PhysicsObject = InitParams.PhysicsObject;
 	Solver = InitParams.Solver;
 	World = InitParams.World;
+	bIsUsingAsyncPhysics = InitParams.bIsUsingAsyncPhysics;
 
 	CurrentSyncState = InitParams.InitialSyncState;
 
@@ -471,7 +472,6 @@ void UChaosMoverSimulation::SimulationTick(const FMoverSimContext& SimContext, c
 
 		RollbackBlackboard->EndSimulationFrame();
 	}
-
 }
 
 bool UChaosMoverSimulation::ShouldResim() const
diff --git a/Engine/Plugins/Experimental/ChaosMover/Source/ChaosMover/Private/PathedMovement/ChaosPathedMovementMode.cpp b/Engine/Plugins/Experimental/ChaosMover/Source/ChaosMover/Private/PathedMovement/ChaosPathedMovementMode.cpp
--- a/Engine/Plugins/Experimental/ChaosMover/Source/ChaosMover/Private/PathedMovement/ChaosPathedMovementMode.cpp
+++ b/Engine/Plugins/Experimental/ChaosMover/Source/ChaosMover/Private/PathedMovement/ChaosPathedMovementMode.cpp
@@ -36,11 +36,24 @@ void UChaosPathedMovementMode::OnRegistered(const FName ModeName, const FMoverSi
 
 	Super::OnRegistered(ModeName, SimContext);
 
+	if (const UChaosMoverSimulation* Sim = GetSimulation())
+	{
+		bIsUsingAsyncPhysics = Sim->IsUsingAsyncPhysics();
+	}
+
 	InitializePath();
 
 	ImmediatePhysics::UpdateJointSettingsFromConstraintProfile(JointConstraintProperties, JointSettings);
 }
 
+FMoverTime UChaosPathedMovementMode::MakeChangeTime(const FMoverTimeStep& TimeStep) const
+{
+	FMoverTime T;
+	T.FrameCount = bIsUsingAsyncPhysics ? TimeStep.ServerFrame : INDEX_NONE;
+	T.TimeMs = TimeStep.BaseSimTimeMs;
+	return T;
+}
+
 bool UChaosPathedMovementMode::ShouldUseConstraint() const
 {
 	return bUseJointConstraint;
@@ -79,16 +92,15 @@ void UChaosPathedMovementMode::SimulationTick_Implementation(const FSimulationTi
 	// We initialize the output movement state from the previous movement state
 	OutputPathedMovemenState = PreviousMoveState;
 
-	// Compute the current path progress before we modify movement properties or update LastChangeFrame 
+	// Compute the current path progress before we modify movement properties or update LastChangeTime
 	float CurrentPathProgress = PreviousMoveState.LastChangePathProgress;
 	if (PreviousProperties.bWantsToPlay && !PreviousMoveState.bHasFinished)
 	{
-		const float FrameDuration = Params.TimeStep.StepMs * 0.001f;
-		const int32 PreviousLastChangeFrame = (PreviousMoveState.LastChangeFrame != INDEX_NONE) ? PreviousMoveState.LastChangeFrame : 0;
-		const int32 FramesSinceLastChange = Params.TimeStep.ServerFrame - PreviousMoveState.LastChangeFrame;
 		const float PlaybackSpeed = (OneWayPlaybackDuration >= 0.0f) ? (1.0f / OneWayPlaybackDuration) : 0.0f;
 		const float PreviousDirection = PreviousMoveState.bIsPathProgressionIncreasing ? 1.0f : -1.0f;
-		CurrentPathProgress += FramesSinceLastChange * (PreviousDirection * FrameDuration * PlaybackSpeed);
+		const double FrameDt = Params.TimeStep.StepMs * 0.001;
+		const double Elapsed = PreviousMoveState.LastChangeTime.ElapsedSecondsTo(Params.TimeStep.ToStartTime(), FrameDt);
+		CurrentPathProgress += (float)(Elapsed * PreviousDirection * PlaybackSpeed);
 	}
 	
 	bool bReceivedNewStartRequest = false;
@@ -109,7 +121,7 @@ void UChaosPathedMovementMode::SimulationTick_Implementation(const FSimulationTi
 		// This is all that should be needed to be able to compute path progress for future frames, until movement properties change again
 		// From now on all frame differences should be using these values and not the ones in PreviousMoveState
 		OutputPathedMovemenState.LastChangePathProgress = CurrentPathProgress;
-		OutputPathedMovemenState.LastChangeFrame = Params.TimeStep.ServerFrame;
+		OutputPathedMovemenState.LastChangeTime = MakeChangeTime(Params.TimeStep);
 
 		bReceivedNewStartRequest = OutputPathedMovemenState.PropertiesInEffect.bWantsToPlay
 			&& (!PreviousProperties.bWantsToPlay
@@ -158,11 +170,11 @@ void UChaosPathedMovementMode::SimulationTick_Implementation(const FSimulationTi
 	float NewPathProgress = CurrentPathProgress;
 	if (CurrentProperties.bWantsToPlay && (!PreviousMoveState.bHasFinished || bReceivedNewStartRequest))
 	{
-		const float FrameDuration = Params.TimeStep.StepMs * 0.001f;
 		const float PlaybackSpeed = (OneWayPlaybackDuration >= 0.0f) ? (1.0f / FMath::Max(OneWayPlaybackDuration, SMALL_NUMBER)) : 0.0f;
 		const float Direction = OutputPathedMovemenState.bIsPathProgressionIncreasing ? 1.0f : -1.0f;
-		const int32 NumFramesSinceLastChange = (Params.TimeStep.ServerFrame - OutputPathedMovemenState.LastChangeFrame) + 1;
-		NewPathProgress = OutputPathedMovemenState.LastChangePathProgress + NumFramesSinceLastChange * (Direction * PlaybackSpeed * FrameDuration);
+		const double FrameDt = Params.TimeStep.StepMs * 0.001;
+		const double Elapsed = OutputPathedMovemenState.LastChangeTime.ElapsedSecondsTo(Params.TimeStep.ToNextTime(), FrameDt);
+		NewPathProgress = OutputPathedMovemenState.LastChangePathProgress + (float)(Elapsed * Direction * PlaybackSpeed);
 
 		if (CurrentProperties.bWantsOneWayPlayback)
 		{
@@ -176,7 +188,7 @@ void UChaosPathedMovementMode::SimulationTick_Implementation(const FSimulationTi
 						NewPathProgress = FMath::Fmod(NewPathProgress, 1.0f);
 						bJustRestarted = true;
 						OutputPathedMovemenState.LastChangePathProgress = NewPathProgress;
-						OutputPathedMovemenState.LastChangeFrame = Params.TimeStep.ServerFrame;
+						OutputPathedMovemenState.LastChangeTime = MakeChangeTime(Params.TimeStep);
 					}
 					else
 					{
@@ -184,7 +196,7 @@ void UChaosPathedMovementMode::SimulationTick_Implementation(const FSimulationTi
 						bHasJustFinished = true;
 						OutputPathedMovemenState.bHasFinished = true;
 						OutputPathedMovemenState.LastChangePathProgress = NewPathProgress;
-						OutputPathedMovemenState.LastChangeFrame = Params.TimeStep.ServerFrame;
+						OutputPathedMovemenState.LastChangeTime = MakeChangeTime(Params.TimeStep);
 					}
 				}
 			}
@@ -197,7 +209,7 @@ void UChaosPathedMovementMode::SimulationTick_Implementation(const FSimulationTi
 						NewPathProgress = 1.0f - FMath::Fmod(-NewPathProgress, 1.0f);
 						bJustRestarted = true;
 						OutputPathedMovemenState.LastChangePathProgress = NewPathProgress;
-						OutputPathedMovemenState.LastChangeFrame = Params.TimeStep.ServerFrame;
+						OutputPathedMovemenState.LastChangeTime = MakeChangeTime(Params.TimeStep);
 					}
 					else
 					{
@@ -205,7 +217,7 @@ void UChaosPathedMovementMode::SimulationTick_Implementation(const FSimulationTi
 						bHasJustFinished = true;
 						OutputPathedMovemenState.bHasFinished = true;
 						OutputPathedMovemenState.LastChangePathProgress = NewPathProgress;
-						OutputPathedMovemenState.LastChangeFrame = Params.TimeStep.ServerFrame;
+						OutputPathedMovemenState.LastChangeTime = MakeChangeTime(Params.TimeStep);
 					}
 				}
 			}
@@ -221,7 +233,7 @@ void UChaosPathedMovementMode::SimulationTick_Implementation(const FSimulationTi
 					NewPathProgress = 1.0f - FMath::Fmod(NewPathProgress, 1.0f);
 					bHasJustBounced = true;
 					OutputPathedMovemenState.LastChangePathProgress = NewPathProgress;
-					OutputPathedMovemenState.LastChangeFrame = Params.TimeStep.ServerFrame;
+					OutputPathedMovemenState.LastChangeTime = MakeChangeTime(Params.TimeStep);
 				}
 				else
 				{
@@ -229,7 +241,7 @@ void UChaosPathedMovementMode::SimulationTick_Implementation(const FSimulationTi
 					bHasJustFinished = true;
 					OutputPathedMovemenState.bHasFinished = true;
 					OutputPathedMovemenState.LastChangePathProgress = NewPathProgress;
-					OutputPathedMovemenState.LastChangeFrame = Params.TimeStep.ServerFrame;
+					OutputPathedMovemenState.LastChangeTime = MakeChangeTime(Params.TimeStep);
 				}
 			}
 			else if (!OutputPathedMovemenState.bIsPathProgressionIncreasing && NewPathProgress <= 0.0f)
@@ -240,7 +252,7 @@ void UChaosPathedMovementMode::SimulationTick_Implementation(const FSimulationTi
 					bHasJustBounced = true;
 					OutputPathedMovemenState.bIsPathProgressionIncreasing = true;
 					OutputPathedMovemenState.LastChangePathProgress = NewPathProgress;
-					OutputPathedMovemenState.LastChangeFrame = Params.TimeStep.ServerFrame;
+					OutputPathedMovemenState.LastChangeTime = MakeChangeTime(Params.TimeStep);
 				}
 				else
 				{
@@ -248,7 +260,7 @@ void UChaosPathedMovementMode::SimulationTick_Implementation(const FSimulationTi
 					bHasJustFinished = true;
 					OutputPathedMovemenState.bHasFinished = true;
 					OutputPathedMovemenState.LastChangePathProgress = NewPathProgress;
-					OutputPathedMovemenState.LastChangeFrame = Params.TimeStep.ServerFrame;
+					OutputPathedMovemenState.LastChangeTime = MakeChangeTime(Params.TimeStep);
 				}
 			}
 		}
diff --git a/Engine/Plugins/Experimental/ChaosMover/Source/ChaosMover/Private/PathedMovement/ChaosPathedMovementTypes.cpp b/Engine/Plugins/Experimental/ChaosMover/Source/ChaosMover/Private/PathedMovement/ChaosPathedMovementTypes.cpp
--- a/Engine/Plugins/Experimental/ChaosMover/Source/ChaosMover/Private/PathedMovement/ChaosPathedMovementTypes.cpp
+++ b/Engine/Plugins/Experimental/ChaosMover/Source/ChaosMover/Private/PathedMovement/ChaosPathedMovementTypes.cpp
@@ -80,13 +80,16 @@ bool FChaosPathedMovementState::NetSerialize(FArchive& Ar, UPackageMap* Map, boo
 
 	bool HasLastChangePathProgress;
 	bool HasLatestMovementStartFrame;
-	bool HasLastChangeFrame;
+	bool HasLastChangeTime;
 
 	if (Ar.IsSaving())
 	{
 		HasLastChangePathProgress = LastChangePathProgress != Defaults.LastChangePathProgress;
 		HasLatestMovementStartFrame = LatestMovementStartFrame != Defaults.LatestMovementStartFrame;
-		HasLastChangeFrame = LastChangeFrame != Defaults.LastChangeFrame;
+		// In variable-dt (non-async) mode, MakeChangeTime sets FrameCount = INDEX_NONE, which matches the default,
+		// so HasLastChangeTime is always false and nothing is sent. This is intentional: variable-dt implies
+		// standalone-only (no networking), so there is nothing to serialize. See note on FChaosPathedMovementInputs.
+		HasLastChangeTime = LastChangeTime.FrameCount != Defaults.LastChangeTime.FrameCount;
 	}
 	else
 	{
@@ -96,7 +99,7 @@ bool FChaosPathedMovementState::NetSerialize(FArchive& Ar, UPackageMap* Map, boo
 
 	Ar.SerializeBits(&HasLastChangePathProgress, 1);
 	Ar.SerializeBits(&HasLatestMovementStartFrame, 1);
-	Ar.SerializeBits(&HasLastChangeFrame, 1);
+	Ar.SerializeBits(&HasLastChangeTime, 1);
 	Ar.SerializeBits(&bIsPathProgressionIncreasing, 1);
 	Ar.SerializeBits(&bHasFinished, 1);
 	PropertiesInEffect.NetSerialize(Ar);
@@ -109,9 +112,9 @@ bool FChaosPathedMovementState::NetSerialize(FArchive& Ar, UPackageMap* Map, boo
 	{
 		Ar << LatestMovementStartFrame;
 	}
-	if (HasLastChangeFrame)
+	if (HasLastChangeTime)
 	{
-		Ar << LastChangeFrame;
+		Ar << LastChangeTime.FrameCount;
 	}
 
 	return true;
@@ -122,7 +125,7 @@ void FChaosPathedMovementState::ToString(FAnsiStringBuilderBase& Out) const
 	FMoverDataStructBase::ToString(Out);
 
 	Out.Appendf("Last Change Path Progress: %.5f | ", LastChangePathProgress);
-	Out.Appendf("Last Change Frame: %d | ", LastChangeFrame);
+	Out.Appendf("Last Change Time: %s | ", *LastChangeTime.ToString());
 	Out.Appendf("Latest Movement Start Frame: %d | ", LatestMovementStartFrame);
 	Out.Appendf("Is Path Progress Increasing: %s | ", bIsPathProgressionIncreasing ? TEXT("True") : TEXT("False"));
 	Out.Appendf("Has Finished: %s | ", bHasFinished ? TEXT("True") : TEXT("False"));
@@ -137,7 +140,7 @@ bool FChaosPathedMovementState::ShouldReconcile(const FMoverDataStructBase& Auth
 
 	if (PropertiesInEffect != TypedAuthority.PropertiesInEffect
 		|| LatestMovementStartFrame != TypedAuthority.LatestMovementStartFrame
-		|| LastChangeFrame != TypedAuthority.LastChangeFrame
+		|| LastChangeTime.FrameCount != TypedAuthority.LastChangeTime.FrameCount
 		|| bIsPathProgressionIncreasing != TypedAuthority.bIsPathProgressionIncreasing
 		|| bHasFinished != TypedAuthority.bHasFinished
 		|| !FMath::IsNearlyEqual(LastChangePathProgress, TypedAuthority.LastChangePathProgress, PathProgressTolerance))
@@ -155,7 +158,7 @@ void FChaosPathedMovementState::Interpolate(const FMoverDataStructBase& From, co
 
 	PropertiesInEffect = (Pct < .5f) ? TypedFrom.PropertiesInEffect : TypedTo.PropertiesInEffect;
 	LatestMovementStartFrame = (Pct < .5f) ? TypedFrom.LatestMovementStartFrame : TypedTo.LatestMovementStartFrame;
-	LastChangeFrame = (Pct < .5f) ? TypedFrom.LastChangeFrame : TypedTo.LastChangeFrame;
+	LastChangeTime = (Pct < .5f) ? TypedFrom.LastChangeTime : TypedTo.LastChangeTime;
 	LastChangePathProgress = FMath::Lerp(TypedFrom.LastChangePathProgress, TypedTo.LastChangePathProgress, Pct);
 	bIsPathProgressionIncreasing = (Pct < .5f) ? TypedFrom.bIsPathProgressionIncreasing : TypedTo.bIsPathProgressionIncreasing;
 	bHasFinished = (Pct < .5f) ? TypedFrom.bHasFinished : TypedTo.bHasFinished;
diff --git a/Engine/Plugins/Experimental/ChaosMover/Source/ChaosMover/Public/ChaosMover/ChaosMoverSimulation.h b/Engine/Plugins/Experimental/ChaosMover/Source/ChaosMover/Public/ChaosMover/ChaosMoverSimulation.h
--- a/Engine/Plugins/Experimental/ChaosMover/Source/ChaosMover/Public/ChaosMover/ChaosMoverSimulation.h
+++ b/Engine/Plugins/Experimental/ChaosMover/Source/ChaosMover/Public/ChaosMover/ChaosMoverSimulation.h
@@ -97,6 +97,7 @@ public:
 	CHAOSMOVER_API bool HasGameplayTag(FGameplayTag TagToFind, bool bExactMatch) const;
 
 	ENetRole GetLocalRole() const { return OwnerLocalRole; }
+	bool IsUsingAsyncPhysics() const { return bIsUsingAsyncPhysics; }
 
 	CHAOSMOVER_API bool ShouldResim() const;
 
@@ -116,6 +117,7 @@ public:
 		Chaos::FConstPhysicsObjectHandle PhysicsObject = nullptr;
 		Chaos::FPhysicsSolver* Solver = nullptr;
 		UWorld* World = nullptr;
+		bool bIsUsingAsyncPhysics = true;
 	};
 
 	CHAOSMOVER_API void Init(const FInitParams& InitParams);
@@ -256,6 +258,7 @@ private:
 
 	bool bInputCmdOverridden = false;
 	bool bSyncStateOverridden = false;
+	bool bIsUsingAsyncPhysics = true;
 
 	int32 InternalServerFrame = 0;
 
diff --git a/Engine/Plugins/Experimental/ChaosMover/Source/ChaosMover/Public/ChaosMover/PathedMovement/ChaosPathedMovementMode.h b/Engine/Plugins/Experimental/ChaosMover/Source/ChaosMover/Public/ChaosMover/PathedMovement/ChaosPathedMovementMode.h
--- a/Engine/Plugins/Experimental/ChaosMover/Source/ChaosMover/Public/ChaosMover/PathedMovement/ChaosPathedMovementMode.h
+++ b/Engine/Plugins/Experimental/ChaosMover/Source/ChaosMover/Public/ChaosMover/PathedMovement/ChaosPathedMovementMode.h
@@ -125,6 +125,14 @@ private:
 
 	Chaos::FPBDJointSettings JointSettings;
 
+	// Whether the physics simulation runs asynchronously with a fixed timestep.
+	// Cached from the simulation at registration time; used only when recording LastChangeTime reference points.
+	bool bIsUsingAsyncPhysics = true;
+
+	// Builds the FMoverTime reference point to store when a change occurs (start of the current frame).
+	// Sets FrameCount in fixed-dt (async) mode, leaves it INDEX_NONE in variable-dt (non-async) mode.
+	FMoverTime MakeChangeTime(const FMoverTimeStep& TimeStep) const;
+
 	/** True to draw debug lines for the aggregate path resulting from all the patterns in the path (i.e. the path that will actually be followed in game) */
 	UPROPERTY(EditAnywhere, Category = "Pathed Movement Mode|Debug")
 	bool bDebugDrawAggregatePath = true;
diff --git a/Engine/Plugins/Experimental/ChaosMover/Source/ChaosMover/Public/ChaosMover/PathedMovement/ChaosPathedMovementTypes.h b/Engine/Plugins/Experimental/ChaosMover/Source/ChaosMover/Public/ChaosMover/PathedMovement/ChaosPathedMovementTypes.h
--- a/Engine/Plugins/Experimental/ChaosMover/Source/ChaosMover/Public/ChaosMover/PathedMovement/ChaosPathedMovementTypes.h
+++ b/Engine/Plugins/Experimental/ChaosMover/Source/ChaosMover/Public/ChaosMover/PathedMovement/ChaosPathedMovementTypes.h
@@ -62,6 +62,12 @@ struct FChaosMutablePathedMovementProperties
  * Frame inputs for all pathed movement
  * Any property that can change during game time must be delivered in this way, only immutable properties can be safely
  * referenced directly on the object they come from (usually the movement mode)
+ *
+ * NOTE: MovementStartFrame and LastChangeFrame use integer frame numbers for scheduling and ordering.
+ * This is correct in async (fixed-dt) mode, where both server and client share the same synchronized frame counter.
+ * In non-async (variable-dt) mode, frame numbers are a local monotonic counter and are NOT synchronized across
+ * endpoints. The current implicit assumption is: networked => fixed-dt, variable-dt => standalone only.
+ * To support variable-dt networking properly, these fields would need to become FMoverTime.
  */
 USTRUCT()
 struct FChaosPathedMovementInputs : public FMoverDataStructBase
@@ -103,11 +109,14 @@ struct FChaosPathedMovementState : public FMoverDataStructBase
 	UPROPERTY(VisibleAnywhere, Category = ChaosMover)
 	float LastChangePathProgress = 0.0f;
 
-	// The frame at which the last change in movement properties has occurred. @see LastChangePathProgress
+	// The time at which the last change in movement properties occurred (set to the start of that frame).
+	// FrameCount is valid in fixed-dt (async) mode; TimeMs is valid in variable-dt (non-async) mode.
+	// @see LastChangePathProgress, FMoverTime::ElapsedSecondsTo
 	UPROPERTY(VisibleAnywhere, Category = ChaosMover)
-	int32 LastChangeFrame = INDEX_NONE;
+	FMoverTime LastChangeTime;
 
-	/** The frame at which we last started playing, or INDEX_NONE if we were never playing */
+	/** The frame at which we last started playing, or INDEX_NONE if we were never playing.
+	 *  Used to order start requests against MovementStartFrame from inputs. See note on FChaosPathedMovementInputs. */
 	UPROPERTY(VisibleAnywhere, Category = ChaosMover)
 	int32 LatestMovementStartFrame = INDEX_NONE;
 
