ChaosMover: Fix to the application of physics object gravity which was computed differently in 2 places in walking mode

#Synced-CL 52710711

diff --git a/Engine/Plugins/Experimental/ChaosMover/Source/ChaosMover/Private/Character/Modes/ChaosWalkingMode.cpp b/Engine/Plugins/Experimental/ChaosMover/Source/ChaosMover/Private/Character/Modes/ChaosWalkingMode.cpp
--- a/Engine/Plugins/Experimental/ChaosMover/Source/ChaosMover/Private/Character/Modes/ChaosWalkingMode.cpp
+++ b/Engine/Plugins/Experimental/ChaosMover/Source/ChaosMover/Private/Character/Modes/ChaosWalkingMode.cpp
@@ -277,8 +277,11 @@ void UChaosWalkingMode::SimulationTick_Implementation(const FSimulationTickParam
 		// Also remove the gravity that will be applied by the physics simulation.
 		// This is so that the gravity in this mode will be consistent with the gravity
 		// set on the mover, not the default physics gravity
+		// Physics volumes always apply gravity along world Z (FVector::UpVector), not along UpDir.
+		const FVector PhysicsObjectGravityVelocity = DefaultSimInputs->PhysicsObjectGravity * FVector::UpVector * DeltaSeconds;
+
 		const FVector ProjectedVelocity = StartingSyncState->GetVelocity_WorldSpace() + DefaultSimInputs->Gravity * DeltaSeconds;
-		FVector TargetVelocity = ProjectedVelocity - DefaultSimInputs->PhysicsObjectGravity * FVector::UpVector * DeltaSeconds;
+		FVector TargetVelocity = ProjectedVelocity - PhysicsObjectGravityVelocity;
 
 		// If we have movement intent and not moving straight up/down then use the proposed move plane velocity
 		// otherwise just fall with gravity
@@ -306,9 +309,7 @@ void UChaosWalkingMode::SimulationTick_Implementation(const FSimulationTickParam
 		FVector ProjectedGroundVelocity = GroundDynamicsInfo.LinearVelocity;
 		if (GroundDynamicsInfo.bIsGravityEnabled)
 		{
-			// This might not be correct if different physics objects have different gravity but is saves having to go
-			// to the component to get the gravity on the physics volume.
-			ProjectedGroundVelocity += DefaultSimInputs->PhysicsObjectGravity * DefaultSimInputs->UpDir * DeltaSeconds;
+			ProjectedGroundVelocity += PhysicsObjectGravityVelocity;
 		}
 
 		const FVector ProjectedRelativeVelocity = TargetVelocity - ProjectedGroundVelocity;
